Added Motion Warping plugin (Experimental)

diff --git a/Engine/Plugins/Experimental/Animation/MotionWarping/MotionWarping.uplugin b/Engine/Plugins/Experimental/Animation/MotionWarping/MotionWarping.uplugin
new file mode 100644
--- /dev/null
+++ b/Engine/Plugins/Experimental/Animation/MotionWarping/MotionWarping.uplugin
@@ -0,0 +1,25 @@
+{
+	"FileVersion": 3,
+	"Version": 1,
+	"VersionName": "0.1",
+	"FriendlyName": "Motion Warping",
+	"Description": "",
+	"Category": "Animation",
+	"CreatedBy": "Epic Games, Inc.",
+	"CreatedByURL": "http://epicgames.com",
+	"DocsURL": "",
+	"MarketplaceURL": "",
+	"SupportURL": "",
+	"EnabledByDefault": false,
+	"CanContainContent": true,
+	"IsBetaVersion": false,
+	"IsExperimentalVersion": true,
+	"Installed": false,
+	"Modules": [
+		{
+			"Name": "MotionWarping",
+			"Type": "Runtime",
+			"LoadingPhase": "Default"
+		}
+	]
+}
\ No newline at end of file
diff --git a/Engine/Plugins/Experimental/Animation/MotionWarping/Source/MotionWarping/MotionWarping.Build.cs b/Engine/Plugins/Experimental/Animation/MotionWarping/Source/MotionWarping/MotionWarping.Build.cs
new file mode 100644
--- /dev/null
+++ b/Engine/Plugins/Experimental/Animation/MotionWarping/Source/MotionWarping/MotionWarping.Build.cs
@@ -0,0 +1,51 @@
+// Copyright Epic Games, Inc. All Rights Reserved.
+
+using UnrealBuildTool;
+
+public class MotionWarping : ModuleRules
+{
+	public MotionWarping(ReadOnlyTargetRules Target) : base(Target)
+	{
+		PCHUsage = ModuleRules.PCHUsageMode.UseExplicitOrSharedPCHs;
+		
+		PublicIncludePaths.AddRange(
+			new string[] {
+				// ... add public include paths required here ...
+			}
+			);
+				
+		
+		PrivateIncludePaths.AddRange(
+			new string[] {
+				// ... add other private include paths required here ...
+			}
+			);
+			
+		
+		PublicDependencyModuleNames.AddRange(
+			new string[]
+			{
+				"Core",
+				"CoreUObject",
+                "Engine",
+			}
+			);
+			
+		
+		PrivateDependencyModuleNames.AddRange(
+			new string[]
+			{
+				"Slate",
+				"SlateCore"
+            }
+			);
+		
+		
+		DynamicallyLoadedModuleNames.AddRange(
+			new string[]
+			{
+				// ... add any modules that your module loads dynamically here ...
+			}
+			);
+	}
+}
diff --git a/Engine/Plugins/Experimental/Animation/MotionWarping/Source/MotionWarping/Private/AnimNotifyState_MotionWarping.cpp b/Engine/Plugins/Experimental/Animation/MotionWarping/Source/MotionWarping/Private/AnimNotifyState_MotionWarping.cpp
new file mode 100644
--- /dev/null
+++ b/Engine/Plugins/Experimental/Animation/MotionWarping/Source/MotionWarping/Private/AnimNotifyState_MotionWarping.cpp
@@ -0,0 +1,9 @@
+// Copyright Epic Games, Inc. All Rights Reserved.
+
+#include "AnimNotifyState_MotionWarping.h"
+#include "GameFramework/Actor.h"
+
+UAnimNotifyState_MotionWarping::UAnimNotifyState_MotionWarping(const FObjectInitializer& ObjectInitializer)
+	: Super(ObjectInitializer)
+{
+}
\ No newline at end of file
diff --git a/Engine/Plugins/Experimental/Animation/MotionWarping/Source/MotionWarping/Private/MotionWarping.cpp b/Engine/Plugins/Experimental/Animation/MotionWarping/Source/MotionWarping/Private/MotionWarping.cpp
new file mode 100644
--- /dev/null
+++ b/Engine/Plugins/Experimental/Animation/MotionWarping/Source/MotionWarping/Private/MotionWarping.cpp
@@ -0,0 +1,19 @@
+// Copyright Epic Games, Inc. All Rights Reserved.
+
+#include "MotionWarping.h"
+
+#define LOCTEXT_NAMESPACE "MotionWarpingModule"
+
+void FMotionWarpingModule::StartupModule()
+{
+
+}
+
+void FMotionWarpingModule::ShutdownModule()
+{
+
+}
+
+#undef LOCTEXT_NAMESPACE
+	
+IMPLEMENT_MODULE(FMotionWarpingModule, MotionWarping)
\ No newline at end of file
diff --git a/Engine/Plugins/Experimental/Animation/MotionWarping/Source/MotionWarping/Private/MotionWarpingComponent.cpp b/Engine/Plugins/Experimental/Animation/MotionWarping/Source/MotionWarping/Private/MotionWarpingComponent.cpp
new file mode 100644
--- /dev/null
+++ b/Engine/Plugins/Experimental/Animation/MotionWarping/Source/MotionWarping/Private/MotionWarpingComponent.cpp
@@ -0,0 +1,262 @@
+// Copyright Epic Games, Inc. All Rights Reserved.
+
+#include "MotionWarpingComponent.h"
+
+#include "RootMotionModifier.h"
+#include "Animation/AnimSequenceBase.h"
+#include "Animation/AnimMontage.h"
+#include "Animation/AnimInstance.h"
+#include "DrawDebugHelpers.h"
+#include "GameFramework/Character.h"
+#include "GameFramework/CharacterMovementComponent.h"
+#include "AnimNotifyState_MotionWarping.h"
+
+DEFINE_LOG_CATEGORY(LogMotionWarping);
+
+#if !(UE_BUILD_SHIPPING || UE_BUILD_TEST)
+TAutoConsoleVariable<int32> FMotionWarpingCVars::CVarMotionWarpingDisable(TEXT("a.MotionWarping.Disable"), 0, TEXT("Disable Motion Warping"), ECVF_Cheat);
+TAutoConsoleVariable<int32> FMotionWarpingCVars::CVarMotionWarpingDebug(TEXT("a.MotionWarping.Debug"), 0, TEXT("0: Disable, 1: Only Log 2: Log and DrawDebug"), ECVF_Cheat);
+TAutoConsoleVariable<float> FMotionWarpingCVars::CVarMotionWarpingDrawDebugDuration(TEXT("a.MotionWarping.DrawDebugLifeTime"), 0.f, TEXT("Time in seconds each draw debug persists.\nRequires 'a.MotionWarping.Debug 2'"), ECVF_Cheat);
+#endif
+
+// UMotionWarpingUtilities
+///////////////////////////////////////////////////////////////////////
+
+void UMotionWarpingUtilities::BreakMotionWarpingSyncPoint(const FMotionWarpingSyncPoint& SyncPoint, FVector& Location, FRotator& Rotation)
+{
+	Location = SyncPoint.GetLocation(); 
+	Rotation = SyncPoint.Rotator();
+}
+
+FMotionWarpingSyncPoint UMotionWarpingUtilities::MakeMotionWarpingSyncPoint(FVector Location, FRotator Rotation)
+{
+	return FMotionWarpingSyncPoint(Location, Rotation);
+}
+
+void UMotionWarpingUtilities::ExtractLocalSpacePose(const UAnimSequenceBase* Animation, const FBoneContainer& BoneContainer, float Time, bool bExtractRootMotion, FCompactPose& OutPose)
+{
+	OutPose.SetBoneContainer(&BoneContainer);
+
+	FBlendedCurve Curve;
+	Curve.InitFrom(BoneContainer);
+
+	FAnimExtractContext Context(Time, bExtractRootMotion);
+
+	FStackCustomAttributes Attributes;
+	FAnimationPoseData AnimationPoseData(OutPose, Curve, Attributes);
+	if (const UAnimSequence* AnimSequence = Cast<UAnimSequence>(Animation))
+	{
+		AnimSequence->GetBonePose(AnimationPoseData, Context);
+	}
+	else if (const UAnimMontage* AnimMontage = Cast<UAnimMontage>(Animation))
+	{
+		const FAnimTrack& AnimTrack = AnimMontage->SlotAnimTracks[0].AnimTrack;
+		AnimTrack.GetAnimationPose(AnimationPoseData, Context);
+	}
+}
+
+void UMotionWarpingUtilities::ExtractComponentSpacePose(const UAnimSequenceBase* Animation, const FBoneContainer& BoneContainer, float Time, bool bExtractRootMotion, FCSPose<FCompactPose>& OutPose)
+{
+	FCompactPose Pose;
+	ExtractLocalSpacePose(Animation, BoneContainer, Time, bExtractRootMotion, Pose);
+	OutPose.InitPose(MoveTemp(Pose));
+}
+
+FTransform UMotionWarpingUtilities::ExtractRootMotionFromAnimation(const UAnimSequenceBase* Animation, float StartTime, float EndTime)
+{
+	if (const UAnimMontage* Anim = Cast<UAnimMontage>(Animation))
+	{
+		return Anim->ExtractRootMotionFromTrackRange(StartTime, EndTime);
+	}
+
+	if (const UAnimSequence* Anim = Cast<UAnimSequence>(Animation))
+	{
+		return Anim->ExtractRootMotionFromRange(StartTime, EndTime);
+	}
+
+	return FTransform::Identity;
+}
+
+// UMotionWarpingComponent
+///////////////////////////////////////////////////////////////////////
+
+UMotionWarpingComponent::UMotionWarpingComponent(const FObjectInitializer& ObjectInitializer)
+	: Super(ObjectInitializer)
+{
+	bWantsInitializeComponent = true;
+}
+
+void UMotionWarpingComponent::InitializeComponent()
+{
+	Super::InitializeComponent();
+
+	CharacterOwner = Cast<ACharacter>(GetOwner());
+
+	UCharacterMovementComponent* CharacterMovementComp = CharacterOwner.IsValid() ? CharacterOwner->GetCharacterMovement() : nullptr;
+	if (CharacterMovementComp)
+	{
+		CharacterMovementComp->ProcessRootMotionPreConvertToWorld.BindUObject(this, &UMotionWarpingComponent::ProcessRootMotionPreConvertToWorld);
+		CharacterMovementComp->ProcessRootMotionPostConvertToWorld.BindUObject(this, &UMotionWarpingComponent::ProcessRootMotionPostConvertToWorld);
+	}
+}
+
+bool UMotionWarpingComponent::ContainsMatchingModifier(const URootMotionModifierConfig* Config, const UAnimSequenceBase* Animation, float StartTime, float EndTime) const
+{
+	if(Config)
+	{
+		return RootMotionModifiers.ContainsByPredicate([=](const TUniquePtr<FRootMotionModifier>& Modifier)
+		{
+			return Config->MatchesConfig(Modifier, Animation, StartTime, EndTime);
+		});
+	}
+
+	return false;
+}
+
+void UMotionWarpingComponent::Update()
+{
+	check(GetCharacterOwner());
+
+	auto TryGetConfigFromNotifyEvent = [](const FAnimNotifyEvent& NotifyEvent)
+	{
+		const UAnimNotifyState_MotionWarping* MotionWarpingNotify = NotifyEvent.NotifyStateClass ? Cast<UAnimNotifyState_MotionWarping>(NotifyEvent.NotifyStateClass) : nullptr;
+		return MotionWarpingNotify ? MotionWarpingNotify->RootMotionModifierConfig : nullptr;
+	};
+
+	const FAnimMontageInstance* RootMotionMontageInstance = GetCharacterOwner()->GetRootMotionAnimMontageInstance();
+	const UAnimMontage* Montage = RootMotionMontageInstance ? RootMotionMontageInstance->Montage : nullptr;
+	if (Montage)
+	{
+		const float PreviousPosition = RootMotionMontageInstance->GetPreviousPosition();
+		const float CurrentPosition = RootMotionMontageInstance->GetPosition();
+
+		// Loop over notifies directly in the montage, looking for Motion Warping windows
+		for (const FAnimNotifyEvent& NotifyEvent : Montage->Notifies)
+		{
+			if (const URootMotionModifierConfig* MotionWarpingConfig = TryGetConfigFromNotifyEvent(NotifyEvent))
+			{
+				const float StartTime = FMath::Clamp(NotifyEvent.GetTriggerTime(), 0.f, Montage->GetPlayLength());
+				const float EndTime = FMath::Clamp(NotifyEvent.GetEndTriggerTime(), 0.f, Montage->GetPlayLength());
+
+				// If we are hitting this window...
+				if (PreviousPosition >= StartTime && PreviousPosition < EndTime)
+				{
+					if (!ContainsMatchingModifier(MotionWarpingConfig, Montage, StartTime, EndTime))
+					{
+						RootMotionModifiers.Add(MotionWarpingConfig->CreateRootMotionModifier(Montage, StartTime, EndTime));
+
+						UE_LOG(LogMotionWarping, Verbose, TEXT("UMotionWarpingComponent::UpdateModifiers Adding RootMotionModifier %s. NetMode: %d WorldTime: %f Animation: %s [%f %f] [%f %f] Loc: %s Rot: %s"),
+							*GetNameSafe(MotionWarpingConfig), GetWorld()->GetNetMode(), GetWorld()->GetTimeSeconds(), *GetNameSafe(Montage), StartTime, EndTime, PreviousPosition, CurrentPosition, 
+							*GetCharacterOwner()->GetActorLocation().ToString(), *GetCharacterOwner()->GetActorRotation().ToCompactString());
+					}
+				}
+			}
+		}
+
+		if(bSearchForWindowsInAnimsWithinMontages)
+		{
+			// Same as before but scanning all animation within the montage
+			for (int32 SlotIdx = 0; SlotIdx < Montage->SlotAnimTracks.Num(); SlotIdx++)
+			{
+				const FAnimTrack& AnimTrack = Montage->SlotAnimTracks[SlotIdx].AnimTrack;
+				const FAnimSegment* AnimSegment = AnimTrack.GetSegmentAtTime(PreviousPosition);
+				if (AnimSegment && AnimSegment->AnimReference)
+				{
+					for (const FAnimNotifyEvent& NotifyEvent : AnimSegment->AnimReference->Notifies)
+					{
+						if (const URootMotionModifierConfig* MotionWarpingConfig = TryGetConfigFromNotifyEvent(NotifyEvent))
+						{
+							const float NotifyStartTime = FMath::Clamp(NotifyEvent.GetTriggerTime(), 0.f, AnimSegment->AnimReference->GetPlayLength());
+							const float NotifyEndTime = FMath::Clamp(NotifyEvent.GetEndTriggerTime(), 0.f, AnimSegment->AnimReference->GetPlayLength());
+
+							// Convert notify times from AnimSequence times to montage times
+							const float StartTime = (NotifyStartTime - AnimSegment->AnimStartTime) + AnimSegment->StartPos;
+							const float EndTime = (NotifyEndTime - AnimSegment->AnimStartTime) + AnimSegment->StartPos;
+
+							if (PreviousPosition >= StartTime && PreviousPosition < EndTime)
+							{
+								if (!ContainsMatchingModifier(MotionWarpingConfig, Montage, StartTime, EndTime))
+								{
+									RootMotionModifiers.Add(MotionWarpingConfig->CreateRootMotionModifier(Montage, StartTime, EndTime));
+
+									UE_LOG(LogMotionWarping, Verbose, TEXT("UMotionWarpingComponent::UpdateModifiers Adding RootMotionModifier %s. NetMode: %d WorldTime: %f Animation: %s [%f %f] [%f %f] Loc: %s Rot: %s"),
+										*GetNameSafe(MotionWarpingConfig), GetWorld()->GetNetMode(), GetWorld()->GetTimeSeconds(), *GetNameSafe(Montage), StartTime, EndTime, PreviousPosition, CurrentPosition,
+										*GetCharacterOwner()->GetActorLocation().ToString(), *GetCharacterOwner()->GetActorRotation().ToCompactString());
+								}
+							}
+						}
+					}
+				}
+			}
+		}
+	}
+
+	// Update the state of all the modifiers
+	for (TUniquePtr<FRootMotionModifier>& Modifier : RootMotionModifiers)
+	{
+		Modifier->Update(*this);
+	}
+
+	// Remove the modifiers that has been marked for removal
+	const int32 TotalRemoved = RootMotionModifiers.RemoveAll([](const TUniquePtr<FRootMotionModifier>& Modifier) { return Modifier->State == ERootMotionModifierState::MarkedForRemoval; });
+
+	UE_CLOG(TotalRemoved > 0, LogMotionWarping, Verbose, TEXT("UMotionWarpingComponent::UpdateModifiers Modifiers removed after update: %d"), TotalRemoved);
+}
+
+FTransform UMotionWarpingComponent::ProcessRootMotionPreConvertToWorld(const FTransform& InRootMotion, UCharacterMovementComponent* CharacterMovementComponent, float DeltaSeconds)
+{
+#if !(UE_BUILD_SHIPPING || UE_BUILD_TEST)
+	if (FMotionWarpingCVars::CVarMotionWarpingDisable.GetValueOnGameThread() > 0)
+	{
+		return InRootMotion;
+	}
+#endif
+
+	// Check for warping windows and update modifier states
+	Update();
+
+	FTransform FinalRootMotion = InRootMotion;
+
+	// Apply Local Space Modifiers
+	for (TUniquePtr<FRootMotionModifier>& RootMotionModifier : RootMotionModifiers)
+	{
+		if (RootMotionModifier->State == ERootMotionModifierState::Active && RootMotionModifier->bInLocalSpace)
+		{
+			FinalRootMotion = RootMotionModifier->ProcessRootMotion(*this, FinalRootMotion, DeltaSeconds);
+		}
+	}
+
+	return FinalRootMotion;
+}
+
+FTransform UMotionWarpingComponent::ProcessRootMotionPostConvertToWorld(const FTransform& InRootMotion, UCharacterMovementComponent* CharacterMovementComponent, float DeltaSeconds)
+{
+#if !(UE_BUILD_SHIPPING || UE_BUILD_TEST)
+	if (FMotionWarpingCVars::CVarMotionWarpingDisable.GetValueOnGameThread() > 0)
+	{
+		return InRootMotion;
+	}
+#endif
+
+	FTransform FinalRootMotion = InRootMotion;
+
+	// Apply World Space Modifiers
+	for (TUniquePtr<FRootMotionModifier>& RootMotionModifier : RootMotionModifiers)
+	{
+		if (RootMotionModifier->State == ERootMotionModifierState::Active && !RootMotionModifier->bInLocalSpace)
+		{
+			FinalRootMotion = RootMotionModifier->ProcessRootMotion(*this, FinalRootMotion, DeltaSeconds);
+		}
+	}
+
+	return FinalRootMotion;
+}
+
+void UMotionWarpingComponent::AddOrUpdateSyncPoint(FName Name, const FMotionWarpingSyncPoint& SyncPoint)
+{
+	if (Name != NAME_None)
+	{
+		FMotionWarpingSyncPoint& MotionWarpingSyncPoint = SyncPoints.FindOrAdd(Name);
+		MotionWarpingSyncPoint = SyncPoint;
+	}
+}
\ No newline at end of file
diff --git a/Engine/Plugins/Experimental/Animation/MotionWarping/Source/MotionWarping/Private/RootMotionModifier.cpp b/Engine/Plugins/Experimental/Animation/MotionWarping/Source/MotionWarping/Private/RootMotionModifier.cpp
new file mode 100644
--- /dev/null
+++ b/Engine/Plugins/Experimental/Animation/MotionWarping/Source/MotionWarping/Private/RootMotionModifier.cpp
@@ -0,0 +1,180 @@
+ // Copyright Epic Games, Inc. All Rights Reserved.
+
+#include "RootMotionModifier.h"
+#include "GameFramework/Character.h"
+#include "GameFramework/CharacterMovementComponent.h"
+#include "Animation/AnimMontage.h"
+#include "MotionWarpingComponent.h"
+#include "DrawDebugHelpers.h"
+
+// FRootMotionModifier
+///////////////////////////////////////////////////////////////
+
+void FRootMotionModifier::Update(UMotionWarpingComponent& OwnerComp)
+{
+	const FAnimMontageInstance* RootMotionMontageInstance = OwnerComp.GetCharacterOwner()->GetRootMotionAnimMontageInstance();
+	const UAnimMontage* Montage = RootMotionMontageInstance ? RootMotionMontageInstance->Montage : nullptr;
+
+	// Mark for removal if our animation is not relevant anymore
+	if (Montage == nullptr || Montage != Animation)
+	{
+		State = ERootMotionModifierState::MarkedForRemoval;
+		return;
+	}
+
+	// Update playback times and weight
+	PreviousPosition = RootMotionMontageInstance->GetPreviousPosition();
+	CurrentPosition = RootMotionMontageInstance->GetPosition();
+	Weight = RootMotionMontageInstance->GetWeight();
+
+	// Mark for removal if the animation already passed the warping window
+	if (PreviousPosition >= EndTime)
+	{
+		State = ERootMotionModifierState::MarkedForRemoval;
+		return;
+	}
+
+	// Check if we are inside the warping window
+	if (PreviousPosition >= StartTime && PreviousPosition < EndTime)
+	{
+		// If we were waiting, switch to active
+		if (State == ERootMotionModifierState::Waiting)
+		{
+			State = ERootMotionModifierState::Active;
+		}
+	}
+}
+
+// FRootMotionModifier_Warp
+///////////////////////////////////////////////////////////////
+
+void FRootMotionModifier_Warp::Update(UMotionWarpingComponent& OwnerComp)
+{
+	// Mark for removal if there is no sync point for us
+	const FMotionWarpingSyncPoint* SyncPointPtr = OwnerComp.FindSyncPoint(SyncPointName);
+	if (SyncPointPtr == nullptr)
+	{
+		State = ERootMotionModifierState::MarkedForRemoval;
+		return;
+	}
+
+	// Update playback times and state
+	FRootMotionModifier::Update(OwnerComp);
+
+	// Cache sync point transform and trigger OnSyncPointChanged if needed
+	if (State == ERootMotionModifierState::Active)
+	{
+		if (CachedSyncPoint != *SyncPointPtr)
+		{
+			CachedSyncPoint = *SyncPointPtr;
+
+			OnSyncPointChanged(OwnerComp);
+		}
+	}
+}
+
+FTransform FRootMotionModifier_Warp::ProcessRootMotion(UMotionWarpingComponent& OwnerComp, const FTransform& InRootMotion, float DeltaSeconds)
+{
+	const ACharacter* CharacterOwner = OwnerComp.GetCharacterOwner();
+	check(CharacterOwner);
+
+	const FTransform& CharacterTransform = CharacterOwner->GetActorTransform();
+
+	FTransform FinalRootMotion = InRootMotion;
+
+	const FTransform RootMotionTotal = UMotionWarpingUtilities::ExtractRootMotionFromAnimation(Animation.Get(), PreviousPosition, EndTime);
+
+	if (bWarpTranslation)
+	{
+		FVector DeltaTranslation = InRootMotion.GetTranslation();
+
+		const FTransform RootMotionDelta = UMotionWarpingUtilities::ExtractRootMotionFromAnimation(Animation.Get(), PreviousPosition, FMath::Min(CurrentPosition, EndTime));
+
+		const float HorizontalDelta = RootMotionDelta.GetTranslation().Size2D();
+		const float HorizontalTarget = FVector::Dist2D(CharacterTransform.GetLocation(), CachedSyncPoint.GetLocation());
+		const float HorizontalOriginal = RootMotionTotal.GetTranslation().Size2D();
+		const float HorizontalTranslationWarped = HorizontalOriginal != 0.f ? ((HorizontalDelta * HorizontalTarget) / HorizontalOriginal) : 0.f;
+
+		if (bInLocalSpace)
+		{
+			const FTransform MeshRelativeTransform = FTransform(CharacterOwner->GetBaseRotationOffset(), CharacterOwner->GetBaseTranslationOffset());
+			const FTransform MeshTransform = MeshRelativeTransform * CharacterOwner->GetActorTransform();
+			DeltaTranslation = MeshTransform.InverseTransformPositionNoScale(CachedSyncPoint.GetLocation()).GetSafeNormal2D() * HorizontalTranslationWarped;
+		}
+		else
+		{
+			DeltaTranslation = (CachedSyncPoint.GetLocation() - CharacterTransform.GetLocation()).GetSafeNormal2D() * HorizontalTranslationWarped;
+		}
+
+		if (!bIgnoreZAxis)
+		{
+			const FVector CapsuleBottomLocation = (CharacterOwner->GetActorLocation() - FVector::UpVector * CharacterOwner->GetSimpleCollisionHalfHeight());
+			const float VerticalDelta = RootMotionDelta.GetTranslation().Z;
+			const float VerticalTarget = CachedSyncPoint.GetLocation().Z - CapsuleBottomLocation.Z;
+			const float VerticalOriginal = RootMotionTotal.GetTranslation().Z;
+			const float VerticalTranslationWarped = VerticalOriginal != 0.f ? ((VerticalDelta * VerticalTarget) / VerticalOriginal) : 0.f;
+
+			DeltaTranslation.Z = VerticalTranslationWarped;
+		}
+
+		FinalRootMotion.SetTranslation(DeltaTranslation);
+	}
+
+	if (bWarpRotation)
+	{
+		const FQuat WarpedRotation = WarpRotation(OwnerComp, RootMotionTotal, InRootMotion, DeltaSeconds);
+		FinalRootMotion.SetRotation(WarpedRotation);
+	}
+
+	// Debug
+#if !(UE_BUILD_SHIPPING || UE_BUILD_TEST)
+	if (FMotionWarpingCVars::CVarMotionWarpingDebug.GetValueOnGameThread() > 0)
+	{
+		PrintLog(OwnerComp, TEXT("FRootMotionModifier_Simple"), InRootMotion, FinalRootMotion);
+	}
+#endif
+
+	return FinalRootMotion;
+}
+
+FQuat FRootMotionModifier_Warp::WarpRotation(UMotionWarpingComponent& OwnerComp, const FTransform& RootMotionDelta, const FTransform& RootMotionTotal, float DeltaSeconds)
+{
+	const FTransform& CharacterTransform = OwnerComp.GetCharacterOwner()->GetActorTransform();
+
+	const FQuat CurrentRotation = CharacterTransform.GetRotation();
+	const float TimeRemaining = EndTime - PreviousPosition;
+	const FQuat Desired = CachedSyncPoint.GetRotation();
+
+	const FQuat RemainingRootRotationInWorld = RootMotionTotal.GetRotation();
+	const FQuat CurrentPlusRemainingRootMotion = RemainingRootRotationInWorld * CurrentRotation;
+	const float PercentThisStep = FMath::Clamp(DeltaSeconds / TimeRemaining, 0.f, 1.f);
+	const FQuat TargetRotThisFrame = FQuat::Slerp(CurrentPlusRemainingRootMotion, Desired, PercentThisStep);
+	const FQuat DeltaOut = TargetRotThisFrame * CurrentPlusRemainingRootMotion.Inverse();
+
+	return (DeltaOut * RootMotionDelta.GetRotation());
+}
+
+#if !(UE_BUILD_SHIPPING || UE_BUILD_TEST)
+void FRootMotionModifier_Warp::PrintLog(const UMotionWarpingComponent& OwnerComp, const FString& Name, const FTransform& OriginalRootMotion, const FTransform& WarpedRootMotion) const
+{
+	const ACharacter* CharacterOwner = OwnerComp.GetCharacterOwner();
+	check(CharacterOwner);
+
+	const FVector CurrentLocation = (CharacterOwner->GetActorLocation() - FVector::UpVector * CharacterOwner->GetSimpleCollisionHalfHeight());
+	const FVector CurrentToTarget = (CachedSyncPoint.GetLocation() - CurrentLocation).GetSafeNormal2D();
+	const FVector FutureLocation = CurrentLocation + (bInLocalSpace ? (CharacterOwner->GetMesh()->ConvertLocalRootMotionToWorld(WarpedRootMotion)).GetTranslation() : WarpedRootMotion.GetTranslation());
+	const FRotator CurrentRotation = CharacterOwner->GetActorRotation();
+	const FRotator FutureRotation = (WarpedRootMotion.GetRotation() * CharacterOwner->GetActorQuat()).Rotator();
+	const float Dot = FVector::DotProduct(CharacterOwner->GetActorForwardVector(), CurrentToTarget);
+	const float CurrentDist2D = FVector::Dist2D(CachedSyncPoint.GetLocation(), CurrentLocation);
+	const float FutureDist2D = FVector::Dist2D(CachedSyncPoint.GetLocation(), FutureLocation);
+	const float DeltaSeconds = CharacterOwner->GetWorld()->GetDeltaSeconds();
+	const float Speed = WarpedRootMotion.GetTranslation().Size() / DeltaSeconds;
+	const float EndTimeOffset = CurrentPosition - EndTime;
+
+	UE_LOG(LogMotionWarping, Log, TEXT("FRootMotionModifier_Simple. NetMode: %d Char: %s Window [%f %f][%f %f] DeltaTime: %f WorldTime: %f EndTimeOffset: %f CurrentDist2D: %f FutureDist2D: %f Dot: %f OriginalMotionDelta: %s (%f) FinalMotionDelta: %s (%f) Speed: %f CurrentLocation: %s FutureLocation: %s CurrentRotation: %s FutureRotation: %s"),
+		(int32)CharacterOwner->GetWorld()->GetNetMode(), *GetNameSafe(CharacterOwner), StartTime, EndTime, PreviousPosition, CurrentPosition, DeltaSeconds, CharacterOwner->GetWorld()->GetTimeSeconds(), EndTimeOffset, CurrentDist2D, FutureDist2D, Dot,
+		*OriginalRootMotion.GetTranslation().ToString(), OriginalRootMotion.GetTranslation().Size(), *WarpedRootMotion.GetTranslation().ToString(), WarpedRootMotion.GetTranslation().Size(), Speed,
+		*CurrentLocation.ToString(), *FutureLocation.ToString(), *CurrentRotation.ToCompactString(), *FutureRotation.ToCompactString());
+}
+#endif
\ No newline at end of file
diff --git a/Engine/Plugins/Experimental/Animation/MotionWarping/Source/MotionWarping/Private/RootMotionModifier_AdjustmentBlendWarp.cpp b/Engine/Plugins/Experimental/Animation/MotionWarping/Source/MotionWarping/Private/RootMotionModifier_AdjustmentBlendWarp.cpp
new file mode 100644
--- /dev/null
+++ b/Engine/Plugins/Experimental/Animation/MotionWarping/Source/MotionWarping/Private/RootMotionModifier_AdjustmentBlendWarp.cpp
@@ -0,0 +1,425 @@
+ // Copyright Epic Games, Inc. All Rights Reserved.
+
+#include "RootMotionModifier_AdjustmentBlendWarp.h"
+
+#include "Animation/AnimSequenceBase.h"
+#include "Animation/AnimMontage.h"
+#include "Animation/AnimInstance.h"
+#include "AnimationUtils.h"
+#include "AnimationRuntime.h"
+#include "DrawDebugHelpers.h"
+#include "GameFramework/Character.h"
+#include "GameFramework/CharacterMovementComponent.h"
+#include "MotionWarpingComponent.h"
+
+DECLARE_CYCLE_STAT(TEXT("MotionWarping PrecomputeWarpedTracks"), STAT_MotionWarping_PrecomputeWarpedTracks, STATGROUP_Anim);
+DECLARE_CYCLE_STAT(TEXT("MotionWarping ExtractMotionDelta"), STAT_MotionWarping_ExtractMotionDelta, STATGROUP_Anim);
+
+FRootMotionModifier_AdjustmentBlendWarp::FRootMotionModifier_AdjustmentBlendWarp()
+{
+	bInLocalSpace = true; 
+}
+
+void FRootMotionModifier_AdjustmentBlendWarp::ExtractBoneTransformAtTime(FTransform& OutTransform, const FName& BoneName, float Time) const
+{
+	const int32 TrackIndex = Result.TrackNames.IndexOfByKey(BoneName);
+	ExtractBoneTransformAtTime(OutTransform, TrackIndex, Time);
+}
+
+void FRootMotionModifier_AdjustmentBlendWarp::ExtractBoneTransformAtTime(FTransform& OutTransform, int32 TrackIndex, float Time) const
+{
+	if (!Result.AnimationTracks.IsValidIndex(TrackIndex))
+	{
+		OutTransform = FTransform::Identity;
+		return;
+	}
+
+	const int32 TotalFrames = FMath::Max(Result.AnimationTracks[TrackIndex].PosKeys.Num(), Result.AnimationTracks[TrackIndex].RotKeys.Num());
+	const float TrackLength = (EndTime - ActualStartTime);
+	const float TimePercent = (Time - ActualStartTime) / (EndTime - ActualStartTime);
+	const float RemappedTime = TimePercent * TrackLength;
+	FAnimationUtils::ExtractTransformFromTrack(RemappedTime, TotalFrames, TrackLength, Result.AnimationTracks[TrackIndex], EAnimInterpolationType::Linear, OutTransform);
+}
+
+void FRootMotionModifier_AdjustmentBlendWarp::ExtractBoneTransformAtFrame(FTransform& OutTransform, int32 TrackIndex, int32 Frame) const
+{
+	if (!Result.AnimationTracks.IsValidIndex(TrackIndex) || !Result.AnimationTracks[TrackIndex].PosKeys.IsValidIndex(Frame))
+	{
+		OutTransform = FTransform::Identity;
+		return;
+	}
+
+	FAnimationUtils::ExtractTransformForFrameFromTrack(Result.AnimationTracks[TrackIndex], Frame, OutTransform);
+}
+
+void FRootMotionModifier_AdjustmentBlendWarp::OnSyncPointChanged(UMotionWarpingComponent& OwnerComp)
+{
+	const ACharacter* CharacterOwner = OwnerComp.GetCharacterOwner();
+	const USkeletalMeshComponent* SkelMeshComp = CharacterOwner->GetMesh();
+
+	ActualStartTime = PreviousPosition;
+	CachedRootMotion = UMotionWarpingUtilities::ExtractRootMotionFromAnimation(Animation.Get(), StartTime, EndTime);
+	CachedMeshRelativeTransform = SkelMeshComp->GetRelativeTransform();
+	CachedMeshTransform = SkelMeshComp->GetComponentTransform();
+
+	Result.AnimationTracks.Reset();
+	Result.TrackNames.Reset();
+}
+
+FTransform FRootMotionModifier_AdjustmentBlendWarp::ProcessRootMotion(UMotionWarpingComponent& OwnerComp, const FTransform& InRootMotion, float DeltaSeconds)
+{
+	// If warped tracks has not been generated yet, do it now
+	if (Result.GetNum() == 0)
+	{
+		PrecomputeWarpedTracks(OwnerComp);
+	}
+
+	// Extract root motion from warped tracks
+	const FTransform FinalRootMotion = ExtractWarpedRootMotion();
+
+	// Debug
+#if !(UE_BUILD_SHIPPING || UE_BUILD_TEST)
+	if (FMotionWarpingCVars::CVarMotionWarpingDebug.GetValueOnGameThread() > 0)
+	{
+		PrintLog(OwnerComp, TEXT("FRootMotionModifier_AdjustmentBlend"), InRootMotion, FinalRootMotion);
+
+		if(FMotionWarpingCVars::CVarMotionWarpingDebug.GetValueOnGameThread() >= 2)
+		{
+			const float DrawDebugDuration = FMotionWarpingCVars::CVarMotionWarpingDrawDebugDuration.GetValueOnGameThread();
+			DrawDebugWarpedTracks(OwnerComp, DrawDebugDuration);
+		}
+	}
+#endif
+
+	return FinalRootMotion;
+}
+
+void FRootMotionModifier_AdjustmentBlendWarp::PrecomputeWarpedTracks(UMotionWarpingComponent& OwnerComp)
+{
+	SCOPE_CYCLE_COUNTER(STAT_MotionWarping_PrecomputeWarpedTracks);
+
+	// First, extract pose at the end of the window for the bones we are going to warp
+
+	const ACharacter* CharacterOwner = OwnerComp.GetCharacterOwner();
+	const FBoneContainer& BoneContainer = CharacterOwner->GetMesh()->GetAnimInstance()->GetRequiredBones();
+
+	// Init FBoneContainer with only the bones that we are interested in
+	TArray<FBoneIndexType> RequiredBoneIndexArray;
+	RequiredBoneIndexArray.Add(0);
+
+	const bool bShouldWarpIKBones = bWarpIKBones && IKBones.Num() > 0;
+	if (bShouldWarpIKBones)
+	{
+		for (const FName& BoneName : IKBones)
+		{
+			const int32 BoneIndex = BoneContainer.GetPoseBoneIndexForBoneName(BoneName);
+			if(BoneIndex != INDEX_NONE)
+			{
+				RequiredBoneIndexArray.Add(BoneIndex);
+			}
+		}
+
+		BoneContainer.GetReferenceSkeleton().EnsureParentsExistAndSort(RequiredBoneIndexArray);
+	}
+
+	// Init BoneContainer
+	FBoneContainer RequiredBones(RequiredBoneIndexArray, FCurveEvaluationOption(false), *BoneContainer.GetAsset());
+
+	// Extract pose
+	FCSPose<FCompactPose> CSPose;
+	UMotionWarpingUtilities::ExtractComponentSpacePose(Animation.Get(), RequiredBones, EndTime, true, CSPose);
+
+	// Second, calculate additive pose
+
+	//Calculate additive translation for root bone
+	FVector RootTargetLocation = CachedMeshTransform.InverseTransformPositionNoScale(CachedSyncPoint.GetLocation());
+
+	FVector RootTotalAdditiveTranslation = FVector::ZeroVector;
+	if (bWarpTranslation)
+	{
+		RootTotalAdditiveTranslation = RootTargetLocation - CachedRootMotion.GetLocation();
+
+		if (bIgnoreZAxis)
+		{
+			RootTotalAdditiveTranslation.Z = 0.f;
+		}
+	}
+
+	// Calculate additive rotation for root bone
+	FQuat RootTotalAdditiveRotation = FQuat::Identity;
+	if (bWarpRotation)
+	{
+		const FQuat TargetRotation = CachedSyncPoint.GetRotation();
+		const FQuat OriginalRotation = CachedMeshRelativeTransform.GetRotation().Inverse() * (CachedRootMotion * CachedMeshTransform).GetRotation();
+		RootTotalAdditiveRotation = FQuat::FindBetweenNormals(OriginalRotation.GetForwardVector(), TargetRotation.GetForwardVector());
+	}
+
+	// Init Additive Pose
+	FCSPose<FCompactPose> AdditivePose;
+	AdditivePose.InitPose(&RequiredBones);
+	AdditivePose.SetComponentSpaceTransform(FCompactPoseBoneIndex(0), FTransform(RootTotalAdditiveRotation, RootTotalAdditiveTranslation));
+
+	// Calculate and add additive pose for IK bones
+	if (bShouldWarpIKBones)
+	{
+		const FTransform RootTargetPoseCS = FTransform((RootTotalAdditiveRotation * CachedRootMotion.GetRotation()), CachedRootMotion.GetTranslation() + RootTotalAdditiveTranslation);
+		for (int32 Idx = 1; Idx < CSPose.GetPose().GetNumBones(); Idx++)
+		{
+			const FName BoneName = RequiredBones.GetReferenceSkeleton().GetBoneName(RequiredBones.GetBoneIndicesArray()[Idx]);
+			if (IKBones.Contains(BoneName))
+			{
+				const int32 BoneIdx = Idx;
+				const FTransform BonePoseCS = CSPose.GetComponentSpaceTransform(FCompactPoseBoneIndex(BoneIdx));
+
+				const FTransform BoneTargetPoseCS = BonePoseCS * RootTargetPoseCS;
+				const FTransform BoneOriginalPoseCS = BonePoseCS * CachedRootMotion;
+
+				const FVector TotalAdditiveTranslation = BoneTargetPoseCS.GetLocation() - BoneOriginalPoseCS.GetLocation();
+
+				AdditivePose.SetComponentSpaceTransform(FCompactPoseBoneIndex(Idx), FTransform(TotalAdditiveTranslation));
+			}
+		}
+	}
+
+	// Finally, run adjustment blending to generate the warped poses for each bone
+
+	//@todo_fer: We could extract and cache this offline when the WarpingWindow is created
+	const float SampleRate = 1 / 60.f;
+	FMotionDeltaTrackContainer MotionDeltaTracks;
+	FRootMotionModifier_AdjustmentBlendWarp::ExtractMotionDeltaFromRange(RequiredBones, Animation.Get(), ActualStartTime, EndTime, SampleRate, MotionDeltaTracks);
+
+	FRootMotionModifier_AdjustmentBlendWarp::AdjustmentBlendWarp(RequiredBones, AdditivePose, MotionDeltaTracks, Result);
+}
+
+void FRootMotionModifier_AdjustmentBlendWarp::ExtractMotionDeltaFromRange(const FBoneContainer& BoneContainer, const UAnimSequenceBase* Animation, float StartTime, float EndTime, float SampleRate, FMotionDeltaTrackContainer& OutMotionDeltaTracks)
+{
+	SCOPE_CYCLE_COUNTER(STAT_MotionWarping_ExtractMotionDelta);
+
+	check(Animation);
+
+	const int32 TotalFrames = FMath::CeilToInt((EndTime - StartTime) / SampleRate) + 1;
+	const int32 TotalBones = BoneContainer.GetCompactPoseNumBones();
+
+	FCompactPose FirstFramePose;
+	UMotionWarpingUtilities::ExtractLocalSpacePose(Animation, BoneContainer, StartTime, false, FirstFramePose);
+	const FTransform RootTransformFirstFrame = FirstFramePose[FCompactPoseBoneIndex(0)];
+
+	FCSPose<FCompactPose> CSPose;
+	UMotionWarpingUtilities::ExtractComponentSpacePose(Animation, BoneContainer, StartTime, false, CSPose);
+
+	OutMotionDeltaTracks.Init(TotalBones);
+	for (int32 BoneIdx = 0; BoneIdx < TotalBones; BoneIdx++)
+	{
+		const FTransform& BoneTransform = (CSPose.GetComponentSpaceTransform(FCompactPoseBoneIndex(BoneIdx)).GetRelativeTransform(RootTransformFirstFrame));
+
+		FMotionDeltaTrack& NewTrack = OutMotionDeltaTracks.Tracks.AddDefaulted_GetRef();
+
+		NewTrack.BoneTransformTrack.Reserve(TotalFrames);
+		NewTrack.DeltaTranslationTrack.Reserve(TotalFrames);
+		NewTrack.DeltaRotationTrack.Reserve(TotalFrames);
+
+		NewTrack.BoneTransformTrack.Add(BoneTransform);
+		NewTrack.DeltaTranslationTrack.Add(FVector::ZeroVector);
+		NewTrack.DeltaRotationTrack.Add(FRotator::ZeroRotator);
+
+		NewTrack.TotalTranslation = FVector::ZeroVector;
+		NewTrack.TotalRotation = FRotator::ZeroRotator;
+	}
+
+	float Time = StartTime;
+	while (Time < EndTime)
+	{
+		Time += FMath::Min(SampleRate, EndTime - Time);
+
+		UMotionWarpingUtilities::ExtractComponentSpacePose(Animation, BoneContainer, Time, false, CSPose);
+
+		for (int32 BoneIdx = 0; BoneIdx < TotalBones; BoneIdx++)
+		{
+			FMotionDeltaTrack& Track = OutMotionDeltaTracks.Tracks[BoneIdx];
+
+			const FTransform& BoneTransform = (CSPose.GetComponentSpaceTransform(FCompactPoseBoneIndex(BoneIdx)).GetRelativeTransform(RootTransformFirstFrame));
+
+			const FTransform& LastBoneTransform = Track.BoneTransformTrack.Last();
+
+			Track.BoneTransformTrack.Add(BoneTransform);
+
+			const FVector Translation = BoneTransform.GetTranslation();
+			const FVector LastTranslation = LastBoneTransform.GetTranslation();
+			const FVector DeltaTranslation(FMath::Abs(Translation.X - LastTranslation.X), FMath::Abs(Translation.Y - LastTranslation.Y), FMath::Abs(Translation.Z - LastTranslation.Z));
+			Track.DeltaTranslationTrack.Add(DeltaTranslation);
+			Track.TotalTranslation = Track.TotalTranslation + DeltaTranslation;
+
+			const FRotator Rotation = BoneTransform.GetRotation().Rotator();
+			const FRotator LastRotation = LastBoneTransform.GetRotation().Rotator();
+			const FRotator DeltaRotation(FMath::Abs(Rotation.Pitch - LastRotation.Pitch), FMath::Abs(Rotation.Yaw - LastRotation.Yaw), FMath::Abs(Rotation.Roll - LastRotation.Roll));
+			Track.DeltaRotationTrack.Add(DeltaRotation);
+			Track.TotalRotation = Track.TotalRotation + DeltaRotation;
+		}
+	}
+}
+
+void FRootMotionModifier_AdjustmentBlendWarp::AdjustmentBlendWarp(const FBoneContainer& BoneContainer, const FCSPose<FCompactPose>& AdditivePose, const FMotionDeltaTrackContainer& MotionDeltaTracks, FAnimSequenceTrackContainer& Output)
+{
+	auto CalculateAdditive = [](const FVector& Total, const FVector& Delta, const FVector& Additive, const FVector& PreviousAdditive, float Alpha)
+	{
+		FVector CurrentAdditive = FVector::ZeroVector;
+		for (int32 Idx = 0; Idx < 3; Idx++)
+		{
+			if (!FMath::IsNearlyZero(Total[Idx], 1.f))
+			{
+				const float Percent = Delta[Idx] / Total[Idx];
+				const float AdditiveDelta = FMath::Abs(Additive[Idx]) * Percent;
+				CurrentAdditive[Idx] = (Additive[Idx] > 0.f) ? PreviousAdditive[Idx] + AdditiveDelta : PreviousAdditive[Idx] - AdditiveDelta;
+			}
+			else
+			{
+				CurrentAdditive[Idx] = Additive[Idx] * Alpha;
+			}
+		}
+		return CurrentAdditive;
+	};
+
+	Output.Initialize(BoneContainer.GetCompactPoseNumBones());
+
+	for (const FCompactPoseBoneIndex PoseBoneIndex : AdditivePose.GetPose().ForEachBoneIndex())
+	{
+		const FTransform& AdditiveTransform = AdditivePose.GetPose()[PoseBoneIndex];
+		if (AdditiveTransform.Equals(FTransform::Identity))
+		{
+			continue;
+		}
+
+		const int32 BoneIndex = PoseBoneIndex.GetInt();
+		const FMotionDeltaTrack& MotionDeltaTrack = MotionDeltaTracks.Tracks[BoneIndex];
+		const int32 TotalFrames = MotionDeltaTrack.BoneTransformTrack.Num();
+
+		FRawAnimSequenceTrack& Track = Output.AnimationTracks[BoneIndex];
+
+		Track.PosKeys.Reserve(TotalFrames);
+		Track.RotKeys.Reserve(TotalFrames);
+
+		Track.PosKeys.Add(MotionDeltaTrack.BoneTransformTrack[0].GetLocation());
+		Track.RotKeys.Add(MotionDeltaTrack.BoneTransformTrack[0].GetRotation());
+
+		Track.ScaleKeys.Add(FVector(1));
+
+		Output.TrackNames[BoneIndex] = BoneContainer.GetReferenceSkeleton().GetBoneName(BoneContainer.GetBoneIndicesArray()[BoneIndex]);
+
+		const FVector TotalAdditiveTranslation = AdditiveTransform.GetTranslation();
+		const FVector TotalAdditiveRotation = AdditiveTransform.GetRotation().Rotator().Euler();
+
+		FVector PrevAdditiveTranslation = FVector::ZeroVector;
+		FVector PrevAdditiveRotation = FVector::ZeroVector;
+		for (int32 FrameIdx = 1; FrameIdx < TotalFrames; FrameIdx++)
+		{
+			const FTransform& BoneTransform = MotionDeltaTrack.BoneTransformTrack[FrameIdx];
+			const float Alpha = (FrameIdx / (float)(TotalFrames - 1));
+
+			// Translation Key
+			const FVector DeltaTranslation = MotionDeltaTrack.DeltaTranslationTrack[FrameIdx];
+			const FVector CurrentAdditiveTranslation = CalculateAdditive(MotionDeltaTrack.TotalTranslation, DeltaTranslation, TotalAdditiveTranslation, PrevAdditiveTranslation, Alpha);
+			Output.AnimationTracks[BoneIndex].PosKeys.Add(BoneTransform.GetTranslation() + CurrentAdditiveTranslation);
+			PrevAdditiveTranslation = CurrentAdditiveTranslation;
+
+			// Rotation Key
+			const FVector DeltaRotation = MotionDeltaTrack.DeltaRotationTrack[FrameIdx].Euler();
+			const FVector CurrentAdditiveRotation = CalculateAdditive(MotionDeltaTrack.TotalRotation.Euler(), DeltaRotation, TotalAdditiveRotation, PrevAdditiveRotation, Alpha);
+			Output.AnimationTracks[BoneIndex].RotKeys.Add(FRotator::MakeFromEuler(BoneTransform.GetRotation().Rotator().Euler() + CurrentAdditiveRotation).Quaternion());
+			PrevAdditiveRotation = CurrentAdditiveRotation;
+		}
+	}
+}
+
+FTransform FRootMotionModifier_AdjustmentBlendWarp::ExtractWarpedRootMotion() const
+{
+	FTransform StartRootTransform;
+	ExtractBoneTransformAtTime(StartRootTransform, 0, PreviousPosition);
+
+	FTransform EndRootTransform;
+	ExtractBoneTransformAtTime(EndRootTransform, 0, CurrentPosition);
+
+	return EndRootTransform.GetRelativeTransform(StartRootTransform);
+}
+
+#if !(UE_BUILD_SHIPPING || UE_BUILD_TEST)
+void FRootMotionModifier_AdjustmentBlendWarp::DrawDebugWarpedTracks(UMotionWarpingComponent& OwnerComp, float DrawDuration) const
+{
+	const UWorld* World = OwnerComp.GetWorld();
+	if (World && Result.GetNum() > 0 && PreviousPosition <= EndTime)
+	{
+		FTransform RootStartTransform;
+		ExtractBoneTransformAtFrame(RootStartTransform, 0, 0);
+
+		for (int32 TrackIndex = 0; TrackIndex < Result.GetNum(); TrackIndex++)
+		{
+			const int32 TotalFrames = Result.AnimationTracks[TrackIndex].PosKeys.Num();
+			if (TotalFrames > 1)
+			{
+				for (int32 FrameIdx = 0; FrameIdx < Result.AnimationTracks[TrackIndex].PosKeys.Num() - 1; FrameIdx++)
+				{
+					FTransform StartTransform;
+					ExtractBoneTransformAtFrame(StartTransform, TrackIndex, FrameIdx);
+					StartTransform = StartTransform * RootStartTransform.Inverse() * CachedMeshTransform;
+
+					FTransform EndTransform;
+					ExtractBoneTransformAtFrame(EndTransform, TrackIndex, FrameIdx + 1);
+					EndTransform = EndTransform * RootStartTransform.Inverse() * CachedMeshTransform;
+
+					DrawDebugLine(World, StartTransform.GetLocation(), EndTransform.GetLocation(), FColor::Yellow, false, DrawDuration, 0, 0.5f);
+				}
+			}
+		}
+
+		const FTransform& MeshTransform = OwnerComp.GetCharacterOwner()->GetMesh()->GetComponentTransform();
+		DrawDebugCoordinateSystem(World, CachedMeshTransform.GetLocation(), CachedMeshTransform.Rotator(), 20.f, false, DrawDuration, 0, 1.f);
+		DrawDebugCoordinateSystem(World, MeshTransform.GetLocation(), MeshTransform.Rotator(), 20.f, false, DrawDuration, 0, 1.f);
+		DrawDebugCoordinateSystem(World, CachedSyncPoint.GetLocation(), CachedSyncPoint.Rotator(), 50.f, false, DrawDuration, 0, 1.f);
+		DrawDebugCoordinateSystem(World, (CachedRootMotion * CachedMeshTransform).GetLocation(), (CachedRootMotion * CachedMeshTransform).Rotator(), 50.f, false, DrawDuration, 0, 1.f);
+	}
+}
+#endif
+
+void FRootMotionModifier_AdjustmentBlendWarp::GetIKBoneTransformAndAlpha(FName BoneName, FTransform& OutTransform, float& OutAlpha) const
+{
+	if (Result.GetNum() == 0 || !bWarpIKBones || !IKBones.Contains(BoneName))
+	{
+		OutTransform = FTransform::Identity;
+		OutAlpha = 0.f;
+		return;
+	}
+
+	FTransform RootPrevPosition;
+	ExtractBoneTransformAtTime(RootPrevPosition, 0, 0.f);
+
+	FTransform BoneTransform;
+	ExtractBoneTransformAtTime(BoneTransform, BoneName, PreviousPosition);
+
+	OutTransform = BoneTransform * RootPrevPosition.Inverse() * CachedMeshTransform;
+	OutAlpha = Weight;
+}
+
+void URootMotionModifierConfig_AdjustmentBlendWarp::GetIKBoneTransformAndAlpha(ACharacter* Character, FName BoneName, FTransform& OutTransform, float& OutAlpha)
+{
+	OutTransform = FTransform::Identity;
+	OutAlpha = 0.f;
+
+	const UMotionWarpingComponent* MotionWarpingComp = Character ? Character->FindComponentByClass<UMotionWarpingComponent>() : nullptr;
+	if (MotionWarpingComp)
+	{
+		for (const auto& Modifier : MotionWarpingComp->GetRootMotionModifiers())
+		{
+			if (Modifier.IsValid() && Modifier->State == ERootMotionModifierState::Active && Modifier->GetScriptStruct()->IsChildOf(FRootMotionModifier_AdjustmentBlendWarp::StaticStruct()))
+			{
+				// We must check if the Animation for the modifier is still relevant because in RootMotionFromMontageOnlyMode the montage could be aborted 
+				// but the modifier will remain in the list until the next update
+				const FAnimMontageInstance* RootMotionMontageInstance = Character->GetRootMotionAnimMontageInstance();
+				const UAnimMontage* Montage = RootMotionMontageInstance ? RootMotionMontageInstance->Montage : nullptr;
+				if (Modifier->Animation == Montage)
+				{
+					static_cast<const FRootMotionModifier_AdjustmentBlendWarp*>(Modifier.Get())->GetIKBoneTransformAndAlpha(BoneName, OutTransform, OutAlpha);
+					break;
+				}
+			}
+		}
+	}
+}
\ No newline at end of file
diff --git a/Engine/Plugins/Experimental/Animation/MotionWarping/Source/MotionWarping/Private/RootMotionModifier_SkewWarp.cpp b/Engine/Plugins/Experimental/Animation/MotionWarping/Source/MotionWarping/Private/RootMotionModifier_SkewWarp.cpp
new file mode 100644
--- /dev/null
+++ b/Engine/Plugins/Experimental/Animation/MotionWarping/Source/MotionWarping/Private/RootMotionModifier_SkewWarp.cpp
@@ -0,0 +1,153 @@
+ // Copyright Epic Games, Inc. All Rights Reserved.
+
+#include "RootMotionModifier_SkewWarp.h"
+#include "DrawDebugHelpers.h"
+#include "GameFramework/Character.h"
+#include "GameFramework/CharacterMovementComponent.h"
+#include "MotionWarpingComponent.h"
+
+FTransform FRootMotionModifier_SkewWarp::ProcessRootMotion(UMotionWarpingComponent& OwnerComp, const FTransform& InRootMotion, float DeltaSeconds)
+{
+	const ACharacter* CharacterOwner = OwnerComp.GetCharacterOwner(); 
+	check(CharacterOwner);
+
+	FTransform FinalRootMotion = InRootMotion;
+
+	const FTransform RootMotionTotal = UMotionWarpingUtilities::ExtractRootMotionFromAnimation(Animation.Get(), PreviousPosition, EndTime);
+	const FTransform RootMotionDelta = UMotionWarpingUtilities::ExtractRootMotionFromAnimation(Animation.Get(), PreviousPosition, FMath::Min(CurrentPosition, EndTime));
+
+	if (bWarpTranslation && !RootMotionDelta.GetTranslation().IsNearlyZero())
+	{
+		const FTransform CurrentTransform = FTransform(
+			CharacterOwner->GetActorQuat(),
+			CharacterOwner->GetActorLocation() - FVector::UpVector * CharacterOwner->GetSimpleCollisionHalfHeight()
+		);
+
+		const FTransform MeshRelativeTransform = FTransform(CharacterOwner->GetBaseRotationOffset(), CharacterOwner->GetBaseTranslationOffset());
+		const FTransform MeshTransform = MeshRelativeTransform * CharacterOwner->GetActorTransform();
+		const FTransform RootMotionTotalWorldSpace = RootMotionTotal * MeshTransform;
+		const FTransform RootMotionDeltaWorldSpace = CharacterOwner->GetMesh()->ConvertLocalRootMotionToWorld(RootMotionDelta);
+
+		const FVector CurrentLocation = CurrentTransform.GetLocation();
+		const FQuat CurrentRotation = CurrentTransform.GetRotation();
+		
+		FVector TargetLocation = CachedSyncPoint.GetLocation();
+		if(bIgnoreZAxis)
+		{
+			TargetLocation.Z = CurrentLocation.Z;
+		}
+
+		const FVector Translation = RootMotionDeltaWorldSpace.GetTranslation();
+		const FVector FutureLocation = RootMotionTotalWorldSpace.GetLocation();
+		const FVector CurrentToWorldOffset = TargetLocation - CurrentLocation;
+		const FVector CurrentToRootOffset = FutureLocation - CurrentLocation;
+
+		// Create a matrix we can use to put everything into a space looking straight at RootMotionSyncPosition. "forward" should be the axis along which we want to scale. 
+		FVector ToRootNormalized = CurrentToRootOffset.GetSafeNormal();
+
+		float BestMatchDot = FMath::Abs(FVector::DotProduct(ToRootNormalized, CurrentRotation.GetAxisX()));
+		FMatrix ToRootSyncSpace = FRotationMatrix::MakeFromXZ(ToRootNormalized, CurrentRotation.GetAxisZ());
+
+		float ZDot = FMath::Abs(FVector::DotProduct(ToRootNormalized, CurrentRotation.GetAxisZ()));
+		if (ZDot > BestMatchDot)
+		{
+			ToRootSyncSpace = FRotationMatrix::MakeFromXZ(ToRootNormalized, CurrentRotation.GetAxisX());
+			BestMatchDot = ZDot;
+		}
+
+		float YDot = FMath::Abs(FVector::DotProduct(ToRootNormalized, CurrentRotation.GetAxisY()));
+		if (YDot > BestMatchDot)
+		{
+			ToRootSyncSpace = FRotationMatrix::MakeFromXZ(ToRootNormalized, CurrentRotation.GetAxisZ());
+		}
+
+		// Put everything into RootSyncSpace.
+		const FVector RootMotionInSyncSpace = ToRootSyncSpace.InverseTransformVector(Translation);
+		const FVector CurrentToWorldSync = ToRootSyncSpace.InverseTransformVector(CurrentToWorldOffset);
+		const FVector CurrentToRootMotionSync = ToRootSyncSpace.InverseTransformVector(CurrentToRootOffset);
+
+		FVector CurrentToWorldSyncNorm = CurrentToWorldSync;
+		CurrentToWorldSyncNorm.Normalize();
+
+		FVector CurrentToRootMotionSyncNorm = CurrentToRootMotionSync;
+		CurrentToRootMotionSyncNorm.Normalize();
+
+		// Calculate skew Yaw Angle. 
+		FVector FlatToWorld = FVector(CurrentToWorldSyncNorm.X, CurrentToWorldSyncNorm.Y, 0.0f);
+		FlatToWorld.Normalize();
+		FVector FlatToRoot = FVector(CurrentToRootMotionSyncNorm.X, CurrentToRootMotionSyncNorm.Y, 0.0f);
+		FlatToRoot.Normalize();
+		float AngleAboutZ = FMath::Acos(FVector::DotProduct(FlatToWorld, FlatToRoot));
+		float AngleAboutZNorm = FMath::DegreesToRadians(FRotator::NormalizeAxis(FMath::RadiansToDegrees(AngleAboutZ)));
+		if (FlatToWorld.Y < 0.0f)
+		{
+			AngleAboutZNorm *= -1.0f;
+		}
+
+		// Calculate Skew Pitch Angle. 
+		FVector ToWorldNoY = FVector(CurrentToWorldSyncNorm.X, 0.0f, CurrentToWorldSyncNorm.Z);
+		ToWorldNoY.Normalize();
+		FVector ToRootNoY = FVector(CurrentToRootMotionSyncNorm.X, 0.0f, CurrentToRootMotionSyncNorm.Z);
+		ToRootNoY.Normalize();
+		const float AngleAboutY = FMath::Acos(FVector::DotProduct(ToWorldNoY, ToRootNoY));
+		float AngleAboutYNorm = FMath::DegreesToRadians(FRotator::NormalizeAxis(FMath::RadiansToDegrees(AngleAboutY)));
+		if (ToWorldNoY.Z < 0.0f)
+		{
+			AngleAboutYNorm *= -1.0f;
+		}
+
+		FVector SkewedRootMotion = FVector::ZeroVector;
+		float ProjectedScale = FVector::DotProduct(CurrentToWorldSync, CurrentToRootMotionSyncNorm) / CurrentToRootMotionSync.Size();
+		if (ProjectedScale != 0.0f)
+		{
+			FMatrix ScaleMatrix;
+			ScaleMatrix.SetIdentity();
+			ScaleMatrix.SetAxis(0, FVector(ProjectedScale, 0.0f, 0.0f));
+			ScaleMatrix.SetAxis(1, FVector(0.0f, 1.0f, 0.0f));
+			ScaleMatrix.SetAxis(2, FVector(0.0f, 0.0f, 1.0f));
+
+			FMatrix ShearXAlongYMatrix;
+			ShearXAlongYMatrix.SetIdentity();
+			ShearXAlongYMatrix.SetAxis(0, FVector(1.0f, FMath::Tan(AngleAboutZNorm), 0.0f));
+			ShearXAlongYMatrix.SetAxis(1, FVector(0.0f, 1.0f, 0.0f));
+			ShearXAlongYMatrix.SetAxis(2, FVector(0.0f, 0.0f, 1.0f));
+
+			FMatrix ShearXAlongZMatrix;
+			ShearXAlongZMatrix.SetIdentity();
+			ShearXAlongZMatrix.SetAxis(0, FVector(1.0f, 0.0f, FMath::Tan(AngleAboutYNorm)));
+			ShearXAlongZMatrix.SetAxis(1, FVector(0.0f, 1.0f, 0.0f));
+			ShearXAlongZMatrix.SetAxis(2, FVector(0.0f, 0.0f, 1.0f));
+
+			FMatrix ScaledSkewMatrix = ScaleMatrix * ShearXAlongYMatrix * ShearXAlongZMatrix;
+
+			// Skew and scale the Root motion. 
+			SkewedRootMotion = ScaledSkewMatrix.TransformVector(RootMotionInSyncSpace);
+		}
+		else if (!CurrentToRootMotionSync.IsZero() && !CurrentToWorldSync.IsZero() && !RootMotionInSyncSpace.IsZero())
+		{
+			// Figure out ratio between remaining Root and remaining World. Then project scaled length of current Root onto World.
+			const float Scale = CurrentToWorldSync.Size() / CurrentToRootMotionSync.Size();
+			const float StepTowardTarget = RootMotionInSyncSpace.ProjectOnTo(RootMotionInSyncSpace).Size();
+			SkewedRootMotion = CurrentToWorldSyncNorm * (Scale * StepTowardTarget);
+		}
+
+		// Put our result back in world space.  
+		FinalRootMotion.SetTranslation(ToRootSyncSpace.TransformVector(SkewedRootMotion));
+	}
+
+	if(bWarpRotation)
+	{
+		const FQuat WarpedRotation = WarpRotation(OwnerComp, RootMotionTotal, InRootMotion, DeltaSeconds);
+		FinalRootMotion.SetRotation(WarpedRotation);
+	}
+
+	// Debug
+#if !(UE_BUILD_SHIPPING || UE_BUILD_TEST)
+	if (FMotionWarpingCVars::CVarMotionWarpingDebug.GetValueOnGameThread() > 0)
+	{
+		PrintLog(OwnerComp, TEXT("FRootMotionModifier_Skew"), InRootMotion, FinalRootMotion);
+	}
+#endif
+
+	return FinalRootMotion;
+}
\ No newline at end of file
diff --git a/Engine/Plugins/Experimental/Animation/MotionWarping/Source/MotionWarping/Public/AnimNotifyState_MotionWarping.h b/Engine/Plugins/Experimental/Animation/MotionWarping/Source/MotionWarping/Public/AnimNotifyState_MotionWarping.h
new file mode 100644
--- /dev/null
+++ b/Engine/Plugins/Experimental/Animation/MotionWarping/Source/MotionWarping/Public/AnimNotifyState_MotionWarping.h
@@ -0,0 +1,21 @@
+// Copyright Epic Games, Inc. All Rights Reserved.
+
+#pragma once
+
+#include "CoreMinimal.h"
+#include "Animation/AnimNotifies/AnimNotifyState.h"
+#include "AnimNotifyState_MotionWarping.generated.h"
+
+/** AnimNotifyState used to define a motion warping window in the animation */
+UCLASS(meta = (DisplayName = "Motion Warping"))
+class UAnimNotifyState_MotionWarping : public UAnimNotifyState
+{
+	GENERATED_BODY()
+
+public:
+
+	UPROPERTY(EditAnywhere, Instanced, Category = "Config")
+	class URootMotionModifierConfig* RootMotionModifierConfig;
+
+	UAnimNotifyState_MotionWarping(const FObjectInitializer& ObjectInitializer);
+};
\ No newline at end of file
diff --git a/Engine/Plugins/Experimental/Animation/MotionWarping/Source/MotionWarping/Public/MotionWarping.h b/Engine/Plugins/Experimental/Animation/MotionWarping/Source/MotionWarping/Public/MotionWarping.h
new file mode 100644
--- /dev/null
+++ b/Engine/Plugins/Experimental/Animation/MotionWarping/Source/MotionWarping/Public/MotionWarping.h
@@ -0,0 +1,15 @@
+// Copyright Epic Games, Inc. All Rights Reserved.
+
+#pragma once
+
+#include "CoreMinimal.h"
+#include "Modules/ModuleManager.h"
+
+class FMotionWarpingModule : public IModuleInterface
+{
+public:
+
+	/** IModuleInterface implementation */
+	virtual void StartupModule() override;
+	virtual void ShutdownModule() override;
+};
diff --git a/Engine/Plugins/Experimental/Animation/MotionWarping/Source/MotionWarping/Public/MotionWarpingComponent.h b/Engine/Plugins/Experimental/Animation/MotionWarping/Source/MotionWarping/Public/MotionWarpingComponent.h
new file mode 100644
--- /dev/null
+++ b/Engine/Plugins/Experimental/Animation/MotionWarping/Source/MotionWarping/Public/MotionWarpingComponent.h
@@ -0,0 +1,100 @@
+// Copyright Epic Games, Inc. All Rights Reserved.
+
+#pragma once
+
+#include "CoreMinimal.h"
+#include "Components/ActorComponent.h"
+#include "Kismet/KismetSystemLibrary.h"
+#include "RootMotionModifier.h"
+#include "MotionWarpingComponent.generated.h"
+
+class ACharacter;
+class UAnimSequenceBase;
+class UCharacterMovementComponent;
+class UMotionWarpingComponent;
+
+DECLARE_LOG_CATEGORY_EXTERN(LogMotionWarping, Log, All);
+
+#if !(UE_BUILD_SHIPPING || UE_BUILD_TEST)
+struct FMotionWarpingCVars
+{
+	static TAutoConsoleVariable<int32> CVarMotionWarpingDisable;
+	static TAutoConsoleVariable<int32> CVarMotionWarpingDebug;
+	static TAutoConsoleVariable<float> CVarMotionWarpingDrawDebugDuration;
+};
+#endif
+
+UCLASS()
+class UMotionWarpingUtilities : public UBlueprintFunctionLibrary
+{
+	GENERATED_BODY()
+
+public:
+
+	/** Extracts data from a MotionWarpingSyncPoint */
+	UFUNCTION(BlueprintPure, Category = "Motion Warping", meta = (NativeBreakFunc))
+	static void BreakMotionWarpingSyncPoint(const FMotionWarpingSyncPoint& SyncPoint, FVector& Location, FRotator& Rotation);
+
+	/** Create a MotionWarpingSyncPoint struct */
+	UFUNCTION(BlueprintPure, Category = "Motion Warping", meta = (NativeMakeFunc))
+	static FMotionWarpingSyncPoint MakeMotionWarpingSyncPoint(FVector Location, FRotator Rotation);
+
+	/** Extract bone pose in local space for all bones in BoneContainer. If Animation is a Montage the pose is extracted from the first track */
+	static void ExtractLocalSpacePose(const UAnimSequenceBase* Animation, const FBoneContainer& BoneContainer, float Time, bool bExtractRootMotion, FCompactPose& OutPose);
+
+	/** Extract bone pose in component space for all bones in BoneContainer. If Animation is a Montage the pose is extracted from the first track */
+	static void ExtractComponentSpacePose(const UAnimSequenceBase* Animation, const FBoneContainer& BoneContainer, float Time, bool bExtractRootMotion, FCSPose<FCompactPose>& OutPose);
+
+	/** Extract Root Motion transform from a contiguous position range */
+	static FTransform ExtractRootMotionFromAnimation(const UAnimSequenceBase* Animation, float StartTime, float EndTime);
+};
+
+UCLASS(ClassGroup = Movement, meta = (BlueprintSpawnableComponent))
+class UMotionWarpingComponent : public UActorComponent
+{
+	GENERATED_BODY()
+
+public:
+
+	/** Whether to look inside animations within montage when looking for warping windows */
+	UPROPERTY(EditAnywhere, BlueprintReadWrite, Category = "Config")
+	bool bSearchForWindowsInAnimsWithinMontages;
+
+	UMotionWarpingComponent(const FObjectInitializer& ObjectInitializer);
+
+	virtual void InitializeComponent() override;
+
+	/** Gets the character this component belongs to */
+	FORCEINLINE ACharacter* GetCharacterOwner() const { return CharacterOwner.Get(); }
+
+	/** Returns the list of root motion modifiers */
+	FORCEINLINE const TArray<TUniquePtr<FRootMotionModifier>>& GetRootMotionModifiers() const { return RootMotionModifiers; }
+
+	/** Find the SyncPoint associated with a specified name */
+	FORCEINLINE const FMotionWarpingSyncPoint* FindSyncPoint(const FName& SyncPointName) const { return SyncPoints.Find(SyncPointName); }
+
+	/** Adds or update the sync point associated with a specified name */
+	UFUNCTION(BlueprintCallable, Category = "Motion Warping")
+	void AddOrUpdateSyncPoint(FName Name, const FMotionWarpingSyncPoint& SyncPoint);
+
+	/** Check if we contain a RootMotionModifier that matches the supplied config data */
+	bool ContainsMatchingModifier(const URootMotionModifierConfig* Config, const UAnimSequenceBase* Animation, float StartTime, float EndTime) const;
+
+protected:
+
+	/** Character this component belongs to */
+	UPROPERTY(Transient)
+	TWeakObjectPtr<ACharacter> CharacterOwner;
+
+	/** List of root motion modifiers */
+	TArray<TUniquePtr<FRootMotionModifier>> RootMotionModifiers;
+
+	UPROPERTY(Transient)
+	TMap<FName, FMotionWarpingSyncPoint> SyncPoints;
+
+	void Update();
+
+	FTransform ProcessRootMotionPreConvertToWorld(const FTransform& InRootMotion, class UCharacterMovementComponent* CharacterMovementComponent, float DeltaSeconds);
+	
+	FTransform ProcessRootMotionPostConvertToWorld(const FTransform& InRootMotion, class UCharacterMovementComponent* CharacterMovementComponent, float DeltaSeconds);
+};
\ No newline at end of file
diff --git a/Engine/Plugins/Experimental/Animation/MotionWarping/Source/MotionWarping/Public/RootMotionModifier.h b/Engine/Plugins/Experimental/Animation/MotionWarping/Source/MotionWarping/Public/RootMotionModifier.h
new file mode 100644
--- /dev/null
+++ b/Engine/Plugins/Experimental/Animation/MotionWarping/Source/MotionWarping/Public/RootMotionModifier.h
@@ -0,0 +1,301 @@
+// Copyright Epic Games, Inc. All Rights Reserved.
+
+#pragma once
+
+#include "CoreMinimal.h"
+#include "RootMotionModifier.generated.h"
+
+class UMotionWarpingComponent;
+class UAnimSequenceBase;
+
+/** The possible states of a Root Motion Modifier */
+UENUM(BlueprintType)
+enum class ERootMotionModifierState : uint8
+{
+	/** The modifier is waiting for the animation to hit the warping window */
+	Waiting,
+
+	/** The modifier is active and currently affecting the final root motion */
+	Active,
+
+	/** The modifier has been marked for removal. Usually because the warping window is done */
+	MarkedForRemoval
+};
+
+/** Base struct for Root Motion Modifiers */
+USTRUCT()
+struct FRootMotionModifier
+{
+	GENERATED_BODY()
+
+public:
+
+	/** Source of the root motion we are warping */
+	UPROPERTY()
+	TWeakObjectPtr<const UAnimSequenceBase> Animation = nullptr;
+
+	/** Start time of the warping window */
+	UPROPERTY()
+	float StartTime = 0.f;
+
+	/** End time of the warping window */
+	UPROPERTY()
+	float EndTime = 0.f;
+
+	/** Previous playback time of the animation */
+	UPROPERTY()
+	float PreviousPosition = 0.f;
+
+	/** Current playback time of the animation */
+	UPROPERTY()
+	float CurrentPosition = 0.f;
+
+	/** Current blend weight of the animation */
+	UPROPERTY()
+	float Weight = 0.f;
+
+	/** Whether this modifier runs before the extracted root motion is converted to world space or after */
+	UPROPERTY()
+	bool bInLocalSpace = false;
+
+	/** Current state */
+	UPROPERTY()
+	ERootMotionModifierState State = ERootMotionModifierState::Waiting;
+
+	FRootMotionModifier(){}
+	virtual ~FRootMotionModifier() {}
+
+	/** Returns the actual struct used for RTTI */
+	virtual UScriptStruct* GetScriptStruct() const PURE_VIRTUAL(FRootMotionModifier::GetScriptStruct, return FRootMotionModifier::StaticStruct(););
+
+	/** Updates the state of the modifier. Runs before ProcessRootMotion */
+	virtual void Update(UMotionWarpingComponent& OwnerComp);
+
+	/** Performs the actual modification to the motion */
+	virtual FTransform ProcessRootMotion(UMotionWarpingComponent& OwnerComp, const FTransform& InRootMotion, float DeltaSeconds) PURE_VIRTUAL(FRootMotionModifier::ProcessRootMotion, return FTransform::Identity;);
+};
+
+/** Blueprint wrapper around the config properties of a root motion modifier */
+UCLASS(abstract, BlueprintType, EditInlineNew)
+class URootMotionModifierConfig : public UObject
+{
+	GENERATED_BODY()
+
+public:
+
+	URootMotionModifierConfig(const FObjectInitializer& ObjectInitializer)
+		: Super(ObjectInitializer) {}
+
+	/** Creates a RootMotionModifier of the type this object represents */
+	virtual TUniquePtr<FRootMotionModifier> CreateRootMotionModifier(const UAnimSequenceBase* Animation, float StartTime, float EndTime) const PURE_VIRTUAL(URootMotionModifierConfig::CreateRootMotionModifier, return TUniquePtr<FRootMotionModifier>(nullptr););
+
+	/** Checks whether the supplied RootMotionModifier matches the config properties in this object. Should be overridden by subclasses to perform the rest of the comparison. */
+	virtual bool MatchesConfig(const TUniquePtr<FRootMotionModifier>& Modifier, const UAnimSequenceBase* Animation, float StartTime, float EndTime) const
+	{
+		return (Modifier.IsValid() && Modifier->Animation == Animation && Modifier->StartTime == StartTime && Modifier->EndTime == EndTime);
+	}
+};
+
+/** Represents a point of alignment in the world */
+USTRUCT(BlueprintType, meta = (HasNativeMake = "MotionWarping.MotionWarpingUtilities.MakeMotionWarpingSyncPoint", HasNativeBreak = "MotionWarping.MotionWarpingUtilities.BreakMotionWarpingSyncPoint"))
+struct FMotionWarpingSyncPoint
+{
+	GENERATED_BODY()
+
+	FMotionWarpingSyncPoint() {}
+	FMotionWarpingSyncPoint(const FVector& InLocation, const FQuat& InRotation)
+		: Location(InLocation), Rotation(InRotation) {}
+	FMotionWarpingSyncPoint(const FVector& InLocation, const FRotator& InRotation)
+		: Location(InLocation), Rotation(InRotation.Quaternion()) {}
+
+	FORCEINLINE const FVector& GetLocation() const { return Location; }
+	FORCEINLINE const FQuat& GetRotation() const { return Rotation; }
+	FORCEINLINE FRotator Rotator() const { return Rotation.Rotator(); }
+
+	FORCEINLINE bool operator==(const FMotionWarpingSyncPoint& Other) const
+	{
+		return Other.Location.Equals(Location) && Other.Rotation.Equals(Rotation);
+	}
+
+	FORCEINLINE bool operator!=(const FMotionWarpingSyncPoint& Other) const
+	{
+		return !Other.Location.Equals(Location) || !Other.Rotation.Equals(Rotation);
+	}
+
+protected:
+
+	UPROPERTY()
+	FVector Location;
+
+	UPROPERTY()
+	FQuat Rotation;
+};
+
+// FRootMotionModifier_Warp
+///////////////////////////////////////////////////////////////
+
+USTRUCT()
+struct FRootMotionModifier_Warp : public FRootMotionModifier
+{
+	GENERATED_BODY()
+
+public:
+
+	/** Name used to find the sync point for this modifier */
+	UPROPERTY()
+	FName SyncPointName = NAME_None;
+
+	/** Whether to warp the translation component of the root motion */
+	UPROPERTY()
+	bool bWarpTranslation = true;
+
+	/** Whether to ignore the Z component of the translation. Z motion will remain untouched */
+	UPROPERTY()
+	bool bIgnoreZAxis = true;
+
+	/** Whether to warp the rotation component of the root motion */
+	UPROPERTY()
+	bool bWarpRotation = true;
+
+	/** Sync Point used by this modifier as target for the warp. Cached during the Update */
+	UPROPERTY()
+	FMotionWarpingSyncPoint CachedSyncPoint;
+
+	FRootMotionModifier_Warp(){}
+	virtual ~FRootMotionModifier_Warp() {}
+
+	//~ Begin FRootMotionModifier Interface
+	virtual UScriptStruct* GetScriptStruct() const override { return FRootMotionModifier_Warp::StaticStruct(); }
+	virtual void Update(UMotionWarpingComponent& OwnerComp) override;
+	virtual FTransform ProcessRootMotion(UMotionWarpingComponent& OwnerComp, const FTransform& InRootMotion, float DeltaSeconds) override;
+	//~ End FRootMotionModifier Interface
+
+	/** Event called during update if the sync point changes while the warping is active */
+	virtual void OnSyncPointChanged(UMotionWarpingComponent& OwnerComp) {}
+
+#if !(UE_BUILD_SHIPPING || UE_BUILD_TEST)
+	void PrintLog(const UMotionWarpingComponent& OwnerComp, const FString& Name, const FTransform& OriginalRootMotion, const FTransform& WarpedRootMotion) const;
+#endif
+
+protected:
+
+	FQuat WarpRotation(UMotionWarpingComponent& OwnerComp, const FTransform& RootMotionDelta, const FTransform& RootMotionTotal, float DeltaSeconds);
+};
+
+UCLASS(meta = (DisplayName = "Simple Warp"))
+class URootMotionModifierConfig_Warp : public URootMotionModifierConfig
+{
+	GENERATED_BODY()
+
+public:
+
+	/** Name used to find the sync point for this modifier */
+	UPROPERTY(EditAnywhere, BlueprintReadWrite, Category = "Config")
+	FName SyncPointName = NAME_None;
+
+	/** Whether to warp the translation component of the root motion */
+	UPROPERTY(EditAnywhere, BlueprintReadWrite, Category = "Config")
+	bool bWarpTranslation = true;
+
+	/** Whether to ignore the Z component of the translation. Z motion will remain untouched */
+	UPROPERTY(EditAnywhere, BlueprintReadWrite, Category = "Config", meta = (EditCondition = "bWarpTranslation"))
+	bool bIgnoreZAxis = true;
+
+	/** Whether to warp the rotation component of the root motion */
+	UPROPERTY(EditAnywhere, BlueprintReadWrite, Category = "Config")
+	bool bWarpRotation = true;
+
+	URootMotionModifierConfig_Warp(const FObjectInitializer& ObjectInitializer) : Super(ObjectInitializer) {}
+
+	virtual TUniquePtr<FRootMotionModifier> CreateRootMotionModifier(const UAnimSequenceBase* Animation, float StartTime, float EndTime) const override
+	{
+		TUniquePtr<FRootMotionModifier_Warp> Modifier = MakeUnique<FRootMotionModifier_Warp>();
+		Modifier->Animation = Animation;
+		Modifier->StartTime = StartTime;
+		Modifier->EndTime = EndTime;
+		Modifier->SyncPointName = SyncPointName;
+		Modifier->bWarpTranslation = bWarpTranslation;
+		Modifier->bIgnoreZAxis = bIgnoreZAxis;
+		Modifier->bWarpRotation = bWarpRotation;
+		return Modifier;
+	}
+
+	virtual bool MatchesConfig(const TUniquePtr<FRootMotionModifier>& Modifier, const UAnimSequenceBase* Animation, float StartTime, float EndTime) const override
+	{
+		if (Super::MatchesConfig(Modifier, Animation, StartTime, EndTime))
+		{
+			if (Modifier->GetScriptStruct()->IsChildOf(FRootMotionModifier_Warp::StaticStruct()))
+			{
+				const FRootMotionModifier_Warp* ModifierAsWarp = static_cast<const FRootMotionModifier_Warp*>(Modifier.Get());
+				return (ModifierAsWarp->SyncPointName == SyncPointName &&
+					ModifierAsWarp->bWarpTranslation == bWarpTranslation &&
+					ModifierAsWarp->bIgnoreZAxis == bIgnoreZAxis &&
+					ModifierAsWarp->bWarpRotation == bWarpRotation);
+			}
+		}
+		return false;
+	}
+};
+
+// FRootMotionModifier_Scale
+///////////////////////////////////////////////////////////////
+
+USTRUCT()
+struct FRootMotionModifier_Scale : public FRootMotionModifier
+{
+	GENERATED_BODY()
+
+public:
+
+	/** Vector used to scale each component of the translation */
+	UPROPERTY()
+	FVector Scale;
+
+	FRootMotionModifier_Scale() { bInLocalSpace = true; }
+	virtual ~FRootMotionModifier_Scale() {}
+
+	virtual UScriptStruct* GetScriptStruct() const override { return FRootMotionModifier_Scale::StaticStruct(); }
+	virtual FTransform ProcessRootMotion(UMotionWarpingComponent& OwnerComp, const FTransform& InRootMotion, float DeltaSeconds) override
+	{
+		FTransform FinalRootMotion = InRootMotion;
+		FinalRootMotion.ScaleTranslation(Scale);
+		return FinalRootMotion;
+	}
+};
+
+UCLASS(meta = (DisplayName = "Scale"))
+class URootMotionModifierConfig_Scale : public URootMotionModifierConfig
+{
+	GENERATED_BODY()
+
+public:
+
+	/** Vector used to scale each component of the translation */
+	UPROPERTY(EditAnywhere, BlueprintReadWrite, Category = "Config")
+	FVector Scale = FVector(1.f);
+
+	URootMotionModifierConfig_Scale(const FObjectInitializer& ObjectInitializer) : Super(ObjectInitializer) {}
+
+	virtual TUniquePtr<FRootMotionModifier> CreateRootMotionModifier(const UAnimSequenceBase* Animation, float StartTime, float EndTime) const override
+	{
+		TUniquePtr<FRootMotionModifier_Scale> Modifier = MakeUnique<FRootMotionModifier_Scale>();
+		Modifier->Animation = Animation;
+		Modifier->StartTime = StartTime;
+		Modifier->EndTime = EndTime;
+		Modifier->Scale = Scale;
+		return Modifier;
+	}
+
+	virtual bool MatchesConfig(const TUniquePtr<FRootMotionModifier>& Modifier, const UAnimSequenceBase* Animation, float StartTime, float EndTime) const override
+	{
+		if (Super::MatchesConfig(Modifier, Animation, StartTime, EndTime))
+		{
+			if (Modifier->GetScriptStruct()->IsChildOf(FRootMotionModifier_Scale::StaticStruct()))
+			{
+				const FRootMotionModifier_Scale* ModifierAsSimple = static_cast<const FRootMotionModifier_Scale*>(Modifier.Get());
+				return (ModifierAsSimple->Scale == Scale);
+			}
+		}
+		return false;
+	}
+};
\ No newline at end of file
diff --git a/Engine/Plugins/Experimental/Animation/MotionWarping/Source/MotionWarping/Public/RootMotionModifier_AdjustmentBlendWarp.h b/Engine/Plugins/Experimental/Animation/MotionWarping/Source/MotionWarping/Public/RootMotionModifier_AdjustmentBlendWarp.h
new file mode 100644
--- /dev/null
+++ b/Engine/Plugins/Experimental/Animation/MotionWarping/Source/MotionWarping/Public/RootMotionModifier_AdjustmentBlendWarp.h
@@ -0,0 +1,133 @@
+// Copyright Epic Games, Inc. All Rights Reserved.
+
+#pragma once
+
+#include "CoreMinimal.h"
+#include "RootMotionModifier.h"
+#include "RootMotionModifier_AdjustmentBlendWarp.generated.h"
+
+USTRUCT()
+struct FMotionDeltaTrack
+{
+	GENERATED_BODY()
+
+	UPROPERTY()
+	TArray<FTransform> BoneTransformTrack;
+
+	UPROPERTY()
+	TArray<FVector> DeltaTranslationTrack;
+
+	UPROPERTY()
+	TArray<FRotator> DeltaRotationTrack;
+
+	UPROPERTY()
+	FVector TotalTranslation;
+
+	UPROPERTY()
+	FRotator TotalRotation;
+};
+
+USTRUCT()
+struct FMotionDeltaTrackContainer
+{
+	GENERATED_BODY()
+
+	UPROPERTY()
+	TArray<FMotionDeltaTrack> Tracks;
+
+	void Init(int32 InNumTracks)
+	{
+		Tracks.Reserve(InNumTracks);
+	}
+};
+
+USTRUCT()
+struct FRootMotionModifier_AdjustmentBlendWarp : public FRootMotionModifier_Warp
+{
+	GENERATED_BODY()
+
+public:
+
+	UPROPERTY()
+	bool bWarpIKBones;
+
+	UPROPERTY()
+	TArray<FName> IKBones;
+
+	FRootMotionModifier_AdjustmentBlendWarp();
+	virtual ~FRootMotionModifier_AdjustmentBlendWarp() {}
+
+	virtual UScriptStruct* GetScriptStruct() const { return FRootMotionModifier_AdjustmentBlendWarp::StaticStruct(); }
+	virtual void OnSyncPointChanged(UMotionWarpingComponent& OwnerComp) override;
+	virtual FTransform ProcessRootMotion(UMotionWarpingComponent& OwnerComp, const FTransform& InRootMotion, float DeltaSeconds) override;
+
+	void GetIKBoneTransformAndAlpha(FName BoneName, FTransform& OutTransform, float& OutAlpha) const;
+
+protected:
+
+	UPROPERTY()
+	FTransform CachedMeshTransform;
+
+	UPROPERTY()
+	FTransform CachedMeshRelativeTransform;
+
+	UPROPERTY()
+	FTransform CachedRootMotion;
+
+	UPROPERTY()
+	FAnimSequenceTrackContainer Result;
+
+	UPROPERTY()
+	float ActualStartTime;
+
+	void PrecomputeWarpedTracks(UMotionWarpingComponent& OwnerComp);
+
+	FTransform ExtractWarpedRootMotion() const;
+
+	void ExtractBoneTransformAtTime(FTransform& OutTransform, const FName& BoneName, float Time) const;
+	void ExtractBoneTransformAtTime(FTransform& OutTransform, int32 TrackIndex, float Time) const;
+	void ExtractBoneTransformAtFrame(FTransform& OutTransform, int32 TrackIndex, int32 Frame) const;
+
+#if !(UE_BUILD_SHIPPING || UE_BUILD_TEST)
+	void DrawDebugWarpedTracks(UMotionWarpingComponent& OwnerComp, float DrawDuration) const;
+#endif
+
+	static void ExtractMotionDeltaFromRange(const FBoneContainer& BoneContainer, const UAnimSequenceBase* Animation, float StartTime, float EndTime, float SampleRate, FMotionDeltaTrackContainer& OutMotionDeltaTracks);
+
+	static void AdjustmentBlendWarp(const FBoneContainer& BoneContainer, const FCSPose<FCompactPose>& AdditivePose, const FMotionDeltaTrackContainer& MotionDeltaTracks, FAnimSequenceTrackContainer& Output);
+};
+
+UCLASS(meta = (DisplayName = "Adjustment Blend Warp"))
+class URootMotionModifierConfig_AdjustmentBlendWarp : public URootMotionModifierConfig_Warp
+{
+	GENERATED_BODY()
+
+public:
+
+	UPROPERTY(EditAnywhere, BlueprintReadWrite, Category = "Config")
+	bool bWarpIKBones;
+
+	UPROPERTY(EditAnywhere, BlueprintReadWrite, Category = "Config", meta = (EditCondition = "bWarpIKBones"))
+	TArray<FName> IKBones;
+
+	URootMotionModifierConfig_AdjustmentBlendWarp(const FObjectInitializer& ObjectInitializer)
+		: Super(ObjectInitializer) {}
+
+	virtual TUniquePtr<FRootMotionModifier> CreateRootMotionModifier(const UAnimSequenceBase* Animation, float StartTime, float EndTime) const override
+	{
+		TUniquePtr<FRootMotionModifier_AdjustmentBlendWarp> Modifier = MakeUnique<FRootMotionModifier_AdjustmentBlendWarp>();
+		Modifier->Animation = Animation;
+		Modifier->StartTime = StartTime;
+		Modifier->EndTime = EndTime;
+		Modifier->SyncPointName = SyncPointName;
+		Modifier->bWarpTranslation = bWarpTranslation;
+		Modifier->bIgnoreZAxis = bIgnoreZAxis;
+		Modifier->bWarpRotation = bWarpRotation;
+		Modifier->bWarpIKBones = bWarpIKBones;
+		Modifier->IKBones = IKBones;
+		return Modifier;
+	}
+
+	UFUNCTION(BlueprintPure, Category = "Motion Warping")
+	static void GetIKBoneTransformAndAlpha(ACharacter* Character, FName BoneName, FTransform& OutTransform, float& OutAlpha);
+};
\ No newline at end of file
diff --git a/Engine/Plugins/Experimental/Animation/MotionWarping/Source/MotionWarping/Public/RootMotionModifier_SkewWarp.h b/Engine/Plugins/Experimental/Animation/MotionWarping/Source/MotionWarping/Public/RootMotionModifier_SkewWarp.h
new file mode 100644
--- /dev/null
+++ b/Engine/Plugins/Experimental/Animation/MotionWarping/Source/MotionWarping/Public/RootMotionModifier_SkewWarp.h
@@ -0,0 +1,45 @@
+// Copyright Epic Games, Inc. All Rights Reserved.
+
+#pragma once
+
+#include "CoreMinimal.h"
+#include "RootMotionModifier.h"
+#include "RootMotionModifier_SkewWarp.generated.h"
+
+USTRUCT()
+struct FRootMotionModifier_SkewWarp : public FRootMotionModifier_Warp
+{
+	GENERATED_BODY()
+
+public:
+
+	FRootMotionModifier_SkewWarp(){}
+	virtual ~FRootMotionModifier_SkewWarp() {}
+
+	virtual UScriptStruct* GetScriptStruct() const { return FRootMotionModifier_SkewWarp::StaticStruct(); }
+	virtual FTransform ProcessRootMotion(UMotionWarpingComponent& OwnerComp, const FTransform& InRootMotion, float DeltaSeconds) override;
+};
+
+UCLASS(meta = (DisplayName = "Skew Warp"))
+class URootMotionModifierConfig_SkewWarp : public URootMotionModifierConfig_Warp
+{
+	GENERATED_BODY()
+
+public:
+
+	URootMotionModifierConfig_SkewWarp(const FObjectInitializer& ObjectInitializer)
+		: Super(ObjectInitializer) {}
+
+	virtual TUniquePtr<FRootMotionModifier> CreateRootMotionModifier(const UAnimSequenceBase* Animation, float StartTime, float EndTime) const override
+	{
+		TUniquePtr<FRootMotionModifier_SkewWarp> Modifier = MakeUnique<FRootMotionModifier_SkewWarp>();
+		Modifier->Animation = Animation;
+		Modifier->StartTime = StartTime;
+		Modifier->EndTime = EndTime;
+		Modifier->SyncPointName = SyncPointName;
+		Modifier->bWarpTranslation = bWarpTranslation;
+		Modifier->bIgnoreZAxis = bIgnoreZAxis;
+		Modifier->bWarpRotation = bWarpRotation;
+		return Modifier;
+	}
+};
\ No newline at end of file
