Added missing include to fix NonUnity compile error.

diff --git a/Engine/Plugins/Experimental/Animation/MotionWarping/Source/MotionWarping/Public/RootMotionModifier.h b/Engine/Plugins/Experimental/Animation/MotionWarping/Source/MotionWarping/Public/RootMotionModifier.h
--- a/Engine/Plugins/Experimental/Animation/MotionWarping/Source/MotionWarping/Public/RootMotionModifier.h
+++ b/Engine/Plugins/Experimental/Animation/MotionWarping/Source/MotionWarping/Public/RootMotionModifier.h
@@ -3,10 +3,10 @@
 #pragma once
 
 #include "CoreMinimal.h"
+#include "Animation/AnimSequence.h"
 #include "RootMotionModifier.generated.h"
 
 class UMotionWarpingComponent;
-class UAnimSequenceBase;
 
 /** The possible states of a Root Motion Modifier */
 UENUM(BlueprintType)
