Added missing includes to fix NonUnity compile error.

diff --git a/Engine/Plugins/Experimental/Animation/MotionWarping/Source/MotionWarping/Private/MotionWarpingComponent.cpp b/Engine/Plugins/Experimental/Animation/MotionWarping/Source/MotionWarping/Private/MotionWarpingComponent.cpp
--- a/Engine/Plugins/Experimental/Animation/MotionWarping/Source/MotionWarping/Private/MotionWarpingComponent.cpp
+++ b/Engine/Plugins/Experimental/Animation/MotionWarping/Source/MotionWarping/Private/MotionWarpingComponent.cpp
@@ -6,6 +6,7 @@
 #include "Animation/AnimSequenceBase.h"
 #include "Animation/AnimMontage.h"
 #include "Animation/AnimInstance.h"
+#include "Animation/AnimationPoseData.h"
 #include "DrawDebugHelpers.h"
 #include "GameFramework/Character.h"
 #include "GameFramework/CharacterMovementComponent.h"
diff --git a/Engine/Plugins/Experimental/Animation/MotionWarping/Source/MotionWarping/Private/RootMotionModifier.cpp b/Engine/Plugins/Experimental/Animation/MotionWarping/Source/MotionWarping/Private/RootMotionModifier.cpp
--- a/Engine/Plugins/Experimental/Animation/MotionWarping/Source/MotionWarping/Private/RootMotionModifier.cpp
+++ b/Engine/Plugins/Experimental/Animation/MotionWarping/Source/MotionWarping/Private/RootMotionModifier.cpp
@@ -3,6 +3,7 @@
 #include "RootMotionModifier.h"
 #include "GameFramework/Character.h"
 #include "GameFramework/CharacterMovementComponent.h"
+#include "Components/SkeletalMeshComponent.h"
 #include "Animation/AnimMontage.h"
 #include "MotionWarpingComponent.h"
 #include "DrawDebugHelpers.h"
diff --git a/Engine/Plugins/Experimental/Animation/MotionWarping/Source/MotionWarping/Private/RootMotionModifier_SkewWarp.cpp b/Engine/Plugins/Experimental/Animation/MotionWarping/Source/MotionWarping/Private/RootMotionModifier_SkewWarp.cpp
--- a/Engine/Plugins/Experimental/Animation/MotionWarping/Source/MotionWarping/Private/RootMotionModifier_SkewWarp.cpp
+++ b/Engine/Plugins/Experimental/Animation/MotionWarping/Source/MotionWarping/Private/RootMotionModifier_SkewWarp.cpp
@@ -4,6 +4,7 @@
 #include "DrawDebugHelpers.h"
 #include "GameFramework/Character.h"
 #include "GameFramework/CharacterMovementComponent.h"
+#include "Components/SkeletalMeshComponent.h"
 #include "MotionWarpingComponent.h"
 
 FTransform FRootMotionModifier_SkewWarp::ProcessRootMotion(UMotionWarpingComponent& OwnerComp, const FTransform& InRootMotion, float DeltaSeconds)
