MotionWarping - Exposed multiple types. Modifiers are now stored as SharedPtrs. Added function to disable all modifiers. Support motion modifier creation in blueprint

diff --git a/Engine/Plugins/Experimental/Animation/MotionWarping/Source/MotionWarping/Private/AnimNotifyState_MotionWarping.cpp b/Engine/Plugins/Experimental/Animation/MotionWarping/Source/MotionWarping/Private/AnimNotifyState_MotionWarping.cpp
--- a/Engine/Plugins/Experimental/Animation/MotionWarping/Source/MotionWarping/Private/AnimNotifyState_MotionWarping.cpp
+++ b/Engine/Plugins/Experimental/Animation/MotionWarping/Source/MotionWarping/Private/AnimNotifyState_MotionWarping.cpp
@@ -2,8 +2,18 @@
 
 #include "AnimNotifyState_MotionWarping.h"
 #include "GameFramework/Actor.h"
+#include "MotionWarpingComponent.h"
+#include "RootMotionModifier.h"
 
 UAnimNotifyState_MotionWarping::UAnimNotifyState_MotionWarping(const FObjectInitializer& ObjectInitializer)
 	: Super(ObjectInitializer)
 {
+}
+
+void UAnimNotifyState_MotionWarping::AddRootMotionModifier_Implementation(UMotionWarpingComponent* MotionWarpingComp, UAnimSequenceBase* Animation, float StartTime, float EndTime) const
+{
+	if (MotionWarpingComp && RootMotionModifierConfig)
+	{
+		RootMotionModifierConfig->AddRootMotionModifier(MotionWarpingComp, Animation, StartTime, EndTime);
+	}
 }
\ No newline at end of file
diff --git a/Engine/Plugins/Experimental/Animation/MotionWarping/Source/MotionWarping/Private/MotionWarpingComponent.cpp b/Engine/Plugins/Experimental/Animation/MotionWarping/Source/MotionWarping/Private/MotionWarpingComponent.cpp
--- a/Engine/Plugins/Experimental/Animation/MotionWarping/Source/MotionWarping/Private/MotionWarpingComponent.cpp
+++ b/Engine/Plugins/Experimental/Animation/MotionWarping/Source/MotionWarping/Private/MotionWarpingComponent.cpp
@@ -101,31 +101,43 @@ void UMotionWarpingComponent::InitializeComponent()
 	}
 }
 
-bool UMotionWarpingComponent::ContainsMatchingModifier(const URootMotionModifierConfig* Config, const UAnimSequenceBase* Animation, float StartTime, float EndTime) const
+bool UMotionWarpingComponent::ContainsModifier(const UAnimSequenceBase* Animation, float StartTime, float EndTime) const
 {
-	if(Config)
-	{
-		return RootMotionModifiers.ContainsByPredicate([=](const TUniquePtr<FRootMotionModifier>& Modifier)
+	return RootMotionModifiers.ContainsByPredicate([=](const TSharedPtr<FRootMotionModifier>& Modifier)
 		{
-			return Config->MatchesConfig(Modifier, Animation, StartTime, EndTime);
+			return (Modifier.IsValid() && Modifier->Animation == Animation && Modifier->StartTime == StartTime && Modifier->EndTime == EndTime);
 		});
+}
+
+void UMotionWarpingComponent::AddRootMotionModifier(TSharedPtr<FRootMotionModifier> Modifier)
+{
+	if (ensureAlways(Modifier.IsValid()))
+	{
+		RootMotionModifiers.Add(Modifier);
+
+		UE_LOG(LogMotionWarping, Verbose, TEXT("UMotionWarpingComponent RootMotionModifier added. NetMode: %d WorldTime: %f Animation: %s [%f %f] [%f %f] Loc: %s Rot: %s"),
+			GetWorld()->GetNetMode(), GetWorld()->GetTimeSeconds(), *GetNameSafe(Modifier->Animation.Get()), Modifier->StartTime, Modifier->EndTime, Modifier->PreviousPosition, Modifier->CurrentPosition,
+			*GetCharacterOwner()->GetActorLocation().ToString(), *GetCharacterOwner()->GetActorRotation().ToCompactString());
 	}
+}
 
-	return false;
+void UMotionWarpingComponent::DisableAllRootMotionModifiers()
+{
+	if (RootMotionModifiers.Num() > 0)
+	{
+		for (TSharedPtr<FRootMotionModifier>& Modifier : RootMotionModifiers)
+		{
+			Modifier->State = ERootMotionModifierState::Disabled;
+		}
+	}
 }
 
 void UMotionWarpingComponent::Update()
 {
 	check(GetCharacterOwner());
 
-	auto TryGetConfigFromNotifyEvent = [](const FAnimNotifyEvent& NotifyEvent)
-	{
-		const UAnimNotifyState_MotionWarping* MotionWarpingNotify = NotifyEvent.NotifyStateClass ? Cast<UAnimNotifyState_MotionWarping>(NotifyEvent.NotifyStateClass) : nullptr;
-		return MotionWarpingNotify ? MotionWarpingNotify->RootMotionModifierConfig : nullptr;
-	};
-
 	const FAnimMontageInstance* RootMotionMontageInstance = GetCharacterOwner()->GetRootMotionAnimMontageInstance();
-	const UAnimMontage* Montage = RootMotionMontageInstance ? RootMotionMontageInstance->Montage : nullptr;
+	UAnimMontage* Montage = RootMotionMontageInstance ? RootMotionMontageInstance->Montage : nullptr;
 	if (Montage)
 	{
 		const float PreviousPosition = RootMotionMontageInstance->GetPreviousPosition();
@@ -134,21 +146,17 @@ void UMotionWarpingComponent::Update()
 		// Loop over notifies directly in the montage, looking for Motion Warping windows
 		for (const FAnimNotifyEvent& NotifyEvent : Montage->Notifies)
 		{
-			if (const URootMotionModifierConfig* MotionWarpingConfig = TryGetConfigFromNotifyEvent(NotifyEvent))
+			const UAnimNotifyState_MotionWarping* MotionWarpingNotify = NotifyEvent.NotifyStateClass ? Cast<UAnimNotifyState_MotionWarping>(NotifyEvent.NotifyStateClass) : nullptr;
+			if (MotionWarpingNotify)
 			{
 				const float StartTime = FMath::Clamp(NotifyEvent.GetTriggerTime(), 0.f, Montage->GetPlayLength());
 				const float EndTime = FMath::Clamp(NotifyEvent.GetEndTriggerTime(), 0.f, Montage->GetPlayLength());
 
-				// If we are hitting this window...
 				if (PreviousPosition >= StartTime && PreviousPosition < EndTime)
 				{
-					if (!ContainsMatchingModifier(MotionWarpingConfig, Montage, StartTime, EndTime))
+					if (!ContainsModifier(Montage, StartTime, EndTime))
 					{
-						RootMotionModifiers.Add(MotionWarpingConfig->CreateRootMotionModifier(Montage, StartTime, EndTime));
-
-						UE_LOG(LogMotionWarping, Verbose, TEXT("UMotionWarpingComponent::UpdateModifiers Adding RootMotionModifier %s. NetMode: %d WorldTime: %f Animation: %s [%f %f] [%f %f] Loc: %s Rot: %s"),
-							*GetNameSafe(MotionWarpingConfig), GetWorld()->GetNetMode(), GetWorld()->GetTimeSeconds(), *GetNameSafe(Montage), StartTime, EndTime, PreviousPosition, CurrentPosition, 
-							*GetCharacterOwner()->GetActorLocation().ToString(), *GetCharacterOwner()->GetActorRotation().ToCompactString());
+						MotionWarpingNotify->AddRootMotionModifier(this, Montage, StartTime, EndTime);
 					}
 				}
 			}
@@ -165,7 +173,9 @@ void UMotionWarpingComponent::Update()
 				{
 					for (const FAnimNotifyEvent& NotifyEvent : AnimSegment->AnimReference->Notifies)
 					{
-						if (const URootMotionModifierConfig* MotionWarpingConfig = TryGetConfigFromNotifyEvent(NotifyEvent))
+						const UAnimNotifyState_MotionWarping* MotionWarpingNotify = NotifyEvent.NotifyStateClass ? Cast<UAnimNotifyState_MotionWarping>(NotifyEvent.NotifyStateClass) : nullptr;
+						const URootMotionModifierConfig* MotionWarpingConfig = MotionWarpingNotify ? MotionWarpingNotify->RootMotionModifierConfig : nullptr;
+						if (MotionWarpingConfig)
 						{
 							const float NotifyStartTime = FMath::Clamp(NotifyEvent.GetTriggerTime(), 0.f, AnimSegment->AnimReference->GetPlayLength());
 							const float NotifyEndTime = FMath::Clamp(NotifyEvent.GetEndTriggerTime(), 0.f, AnimSegment->AnimReference->GetPlayLength());
@@ -176,13 +186,9 @@ void UMotionWarpingComponent::Update()
 
 							if (PreviousPosition >= StartTime && PreviousPosition < EndTime)
 							{
-								if (!ContainsMatchingModifier(MotionWarpingConfig, Montage, StartTime, EndTime))
+								if (!ContainsModifier(Montage, StartTime, EndTime))
 								{
-									RootMotionModifiers.Add(MotionWarpingConfig->CreateRootMotionModifier(Montage, StartTime, EndTime));
-
-									UE_LOG(LogMotionWarping, Verbose, TEXT("UMotionWarpingComponent::UpdateModifiers Adding RootMotionModifier %s. NetMode: %d WorldTime: %f Animation: %s [%f %f] [%f %f] Loc: %s Rot: %s"),
-										*GetNameSafe(MotionWarpingConfig), GetWorld()->GetNetMode(), GetWorld()->GetTimeSeconds(), *GetNameSafe(Montage), StartTime, EndTime, PreviousPosition, CurrentPosition,
-										*GetCharacterOwner()->GetActorLocation().ToString(), *GetCharacterOwner()->GetActorRotation().ToCompactString());
+									MotionWarpingNotify->AddRootMotionModifier(this, Montage, StartTime, EndTime);
 								}
 							}
 						}
@@ -192,16 +198,21 @@ void UMotionWarpingComponent::Update()
 		}
 	}
 
+	OnPreUpdate.Broadcast(this);
+
 	// Update the state of all the modifiers
-	for (TUniquePtr<FRootMotionModifier>& Modifier : RootMotionModifiers)
+	if(RootMotionModifiers.Num() > 0)
 	{
-		Modifier->Update(*this);
-	}
+		for (TSharedPtr<FRootMotionModifier>& Modifier : RootMotionModifiers)
+		{
+			Modifier->Update(*this);
+		}
 
-	// Remove the modifiers that has been marked for removal
-	const int32 TotalRemoved = RootMotionModifiers.RemoveAll([](const TUniquePtr<FRootMotionModifier>& Modifier) { return Modifier->State == ERootMotionModifierState::MarkedForRemoval; });
+		// Remove the modifiers that has been marked for removal
+		const int32 TotalRemoved = RootMotionModifiers.RemoveAll([](const TSharedPtr<FRootMotionModifier>& Modifier) { return Modifier->State == ERootMotionModifierState::MarkedForRemoval; });
 
-	UE_CLOG(TotalRemoved > 0, LogMotionWarping, Verbose, TEXT("UMotionWarpingComponent::UpdateModifiers Modifiers removed after update: %d"), TotalRemoved);
+		UE_CLOG(TotalRemoved > 0, LogMotionWarping, Verbose, TEXT("UMotionWarpingComponent::UpdateModifiers Modifiers removed after update: %d"), TotalRemoved);
+	}
 }
 
 FTransform UMotionWarpingComponent::ProcessRootMotionPreConvertToWorld(const FTransform& InRootMotion, UCharacterMovementComponent* CharacterMovementComponent, float DeltaSeconds)
@@ -219,7 +230,7 @@ FTransform UMotionWarpingComponent::ProcessRootMotionPreConvertToWorld(const FTr
 	FTransform FinalRootMotion = InRootMotion;
 
 	// Apply Local Space Modifiers
-	for (TUniquePtr<FRootMotionModifier>& RootMotionModifier : RootMotionModifiers)
+	for (TSharedPtr<FRootMotionModifier>& RootMotionModifier : RootMotionModifiers)
 	{
 		if (RootMotionModifier->State == ERootMotionModifierState::Active && RootMotionModifier->bInLocalSpace)
 		{
@@ -242,7 +253,7 @@ FTransform UMotionWarpingComponent::ProcessRootMotionPostConvertToWorld(const FT
 	FTransform FinalRootMotion = InRootMotion;
 
 	// Apply World Space Modifiers
-	for (TUniquePtr<FRootMotionModifier>& RootMotionModifier : RootMotionModifiers)
+	for (TSharedPtr<FRootMotionModifier>& RootMotionModifier : RootMotionModifiers)
 	{
 		if (RootMotionModifier->State == ERootMotionModifierState::Active && !RootMotionModifier->bInLocalSpace)
 		{
diff --git a/Engine/Plugins/Experimental/Animation/MotionWarping/Source/MotionWarping/Private/RootMotionModifier.cpp b/Engine/Plugins/Experimental/Animation/MotionWarping/Source/MotionWarping/Private/RootMotionModifier.cpp
--- a/Engine/Plugins/Experimental/Animation/MotionWarping/Source/MotionWarping/Private/RootMotionModifier.cpp
+++ b/Engine/Plugins/Experimental/Animation/MotionWarping/Source/MotionWarping/Private/RootMotionModifier.cpp
@@ -55,6 +55,7 @@ void FRootMotionModifier_Warp::Update(UMotionWarpingComponent& OwnerComp)
 	const FMotionWarpingSyncPoint* SyncPointPtr = OwnerComp.FindSyncPoint(SyncPointName);
 	if (SyncPointPtr == nullptr)
 	{
+		UE_LOG(LogMotionWarping, Warning, TEXT("Marking RootMotionModifier for removal. Reason: Can't find SyncPoint (%s)"), *SyncPointName.ToString());
 		State = ERootMotionModifierState::MarkedForRemoval;
 		return;
 	}
@@ -173,9 +174,44 @@ void FRootMotionModifier_Warp::PrintLog(const UMotionWarpingComponent& OwnerComp
 	const float Speed = WarpedRootMotion.GetTranslation().Size() / DeltaSeconds;
 	const float EndTimeOffset = CurrentPosition - EndTime;
 
-	UE_LOG(LogMotionWarping, Log, TEXT("FRootMotionModifier_Simple. NetMode: %d Char: %s Window [%f %f][%f %f] DeltaTime: %f WorldTime: %f EndTimeOffset: %f CurrentDist2D: %f FutureDist2D: %f Dot: %f OriginalMotionDelta: %s (%f) FinalMotionDelta: %s (%f) Speed: %f CurrentLocation: %s FutureLocation: %s CurrentRotation: %s FutureRotation: %s"),
-		(int32)CharacterOwner->GetWorld()->GetNetMode(), *GetNameSafe(CharacterOwner), StartTime, EndTime, PreviousPosition, CurrentPosition, DeltaSeconds, CharacterOwner->GetWorld()->GetTimeSeconds(), EndTimeOffset, CurrentDist2D, FutureDist2D, Dot,
+	UE_LOG(LogMotionWarping, Log, TEXT("%s. NetMode: %d Char: %s Anim: %s Window [%f %f][%f %f] DeltaTime: %f WorldTime: %f EndTimeOffset: %f Dist2D: %f FutureDist2D: %f Dot: %f OriginalMotionDelta: %s (%f) FinalMotionDelta: %s (%f) Speed: %f Location: %s FutureLocation: %s Rotation: %s FutureRotation: %s"),
+		*Name, (int32)CharacterOwner->GetWorld()->GetNetMode(), *GetNameSafe(CharacterOwner), *GetNameSafe(Animation.Get()), StartTime, EndTime, PreviousPosition, CurrentPosition, DeltaSeconds, CharacterOwner->GetWorld()->GetTimeSeconds(), EndTimeOffset, CurrentDist2D, FutureDist2D, Dot,
 		*OriginalRootMotion.GetTranslation().ToString(), OriginalRootMotion.GetTranslation().Size(), *WarpedRootMotion.GetTranslation().ToString(), WarpedRootMotion.GetTranslation().Size(), Speed,
 		*CurrentLocation.ToString(), *FutureLocation.ToString(), *CurrentRotation.ToCompactString(), *FutureRotation.ToCompactString());
 }
-#endif
\ No newline at end of file
+#endif
+
+// URootMotionModifierConfig_Warp
+///////////////////////////////////////////////////////////////
+
+void URootMotionModifierConfig_Warp::AddRootMotionModifierSimpleWarp(UMotionWarpingComponent* InMotionWarpingComp, const UAnimSequenceBase* InAnimation, float InStartTime, float InEndTime, FName InSyncPointName, bool bInWarpTranslation, bool bInIgnoreZAxis, bool bInWarpRotation)
+{
+	if (ensureAlways(InMotionWarpingComp))
+	{
+		TSharedPtr<FRootMotionModifier_Warp> NewModifier = MakeShared<FRootMotionModifier_Warp>();
+		NewModifier->Animation = InAnimation;
+		NewModifier->StartTime = InStartTime;
+		NewModifier->EndTime = InEndTime;
+		NewModifier->SyncPointName = InSyncPointName;
+		NewModifier->bWarpTranslation = bInWarpTranslation;
+		NewModifier->bIgnoreZAxis = bInIgnoreZAxis;
+		NewModifier->bWarpRotation = bInWarpRotation;
+		InMotionWarpingComp->AddRootMotionModifier(NewModifier);
+	}
+}
+
+// URootMotionModifierConfig_Scale
+///////////////////////////////////////////////////////////////
+
+void URootMotionModifierConfig_Scale::AddRootMotionModifierScale(UMotionWarpingComponent* InMotionWarpingComp, const UAnimSequenceBase* InAnimation, float InStartTime, float InEndTime, FVector InScale)
+{
+	if (ensureAlways(InMotionWarpingComp))
+	{
+		TSharedPtr<FRootMotionModifier_Scale> NewModifier = MakeShared<FRootMotionModifier_Scale>();
+		NewModifier->Animation = InAnimation;
+		NewModifier->StartTime = InStartTime;
+		NewModifier->EndTime = InEndTime;
+		NewModifier->Scale = InScale;
+		InMotionWarpingComp->AddRootMotionModifier(NewModifier);
+	}
+}
\ No newline at end of file
diff --git a/Engine/Plugins/Experimental/Animation/MotionWarping/Source/MotionWarping/Private/RootMotionModifier_AdjustmentBlendWarp.cpp b/Engine/Plugins/Experimental/Animation/MotionWarping/Source/MotionWarping/Private/RootMotionModifier_AdjustmentBlendWarp.cpp
--- a/Engine/Plugins/Experimental/Animation/MotionWarping/Source/MotionWarping/Private/RootMotionModifier_AdjustmentBlendWarp.cpp
+++ b/Engine/Plugins/Experimental/Animation/MotionWarping/Source/MotionWarping/Private/RootMotionModifier_AdjustmentBlendWarp.cpp
@@ -421,4 +421,25 @@ void URootMotionModifierConfig_AdjustmentBlendWarp::GetIKBoneTransformAndAlpha(A
 			}
 		}
 	}
+}
+
+// URootMotionModifierConfig_AdjustmentBlendWarp
+///////////////////////////////////////////////////////////////
+
+void URootMotionModifierConfig_AdjustmentBlendWarp::AddRootMotionModifierAdjustmentBlendWarp(UMotionWarpingComponent* InMotionWarpingComp, const UAnimSequenceBase* InAnimation, float InStartTime, float InEndTime, FName InSyncPointName, bool bInWarpTranslation, bool bInIgnoreZAxis, bool bInWarpRotation, bool bInWarpIKBones, const TArray<FName>& InIKBones)
+{
+	if (ensureAlways(InMotionWarpingComp))
+	{
+		TSharedPtr<FRootMotionModifier_AdjustmentBlendWarp> NewModifier = MakeShared<FRootMotionModifier_AdjustmentBlendWarp>();
+		NewModifier->Animation = InAnimation;
+		NewModifier->StartTime = InStartTime;
+		NewModifier->EndTime = InEndTime;
+		NewModifier->SyncPointName = InSyncPointName;
+		NewModifier->bWarpTranslation = bInWarpTranslation;
+		NewModifier->bIgnoreZAxis = bInIgnoreZAxis;
+		NewModifier->bWarpRotation = bInWarpRotation;
+		NewModifier->bWarpIKBones = bInWarpIKBones;
+		NewModifier->IKBones = InIKBones;
+		InMotionWarpingComp->AddRootMotionModifier(NewModifier);
+	}
 }
\ No newline at end of file
diff --git a/Engine/Plugins/Experimental/Animation/MotionWarping/Source/MotionWarping/Private/RootMotionModifier_SkewWarp.cpp b/Engine/Plugins/Experimental/Animation/MotionWarping/Source/MotionWarping/Private/RootMotionModifier_SkewWarp.cpp
--- a/Engine/Plugins/Experimental/Animation/MotionWarping/Source/MotionWarping/Private/RootMotionModifier_SkewWarp.cpp
+++ b/Engine/Plugins/Experimental/Animation/MotionWarping/Source/MotionWarping/Private/RootMotionModifier_SkewWarp.cpp
@@ -151,4 +151,23 @@ FTransform FRootMotionModifier_SkewWarp::ProcessRootMotion(UMotionWarpingCompone
 #endif
 
 	return FinalRootMotion;
+}
+
+// URootMotionModifierConfig_SkewWarp
+///////////////////////////////////////////////////////////////
+
+void URootMotionModifierConfig_SkewWarp::AddRootMotionModifierSkewWarp(UMotionWarpingComponent* InMotionWarpingComp, const UAnimSequenceBase* InAnimation, float InStartTime, float InEndTime, FName InSyncPointName, bool bInWarpTranslation, bool bInIgnoreZAxis, bool bInWarpRotation)
+{
+	if (ensureAlways(InMotionWarpingComp))
+	{
+		TSharedPtr<FRootMotionModifier_SkewWarp> NewModifier = MakeShared<FRootMotionModifier_SkewWarp>();
+		NewModifier->Animation = InAnimation;
+		NewModifier->StartTime = InStartTime;
+		NewModifier->EndTime = InEndTime;
+		NewModifier->SyncPointName = InSyncPointName;
+		NewModifier->bWarpTranslation = bInWarpTranslation;
+		NewModifier->bIgnoreZAxis = bInIgnoreZAxis;
+		NewModifier->bWarpRotation = bInWarpRotation;
+		InMotionWarpingComp->AddRootMotionModifier(NewModifier);
+	}
 }
\ No newline at end of file
diff --git a/Engine/Plugins/Experimental/Animation/MotionWarping/Source/MotionWarping/Public/AnimNotifyState_MotionWarping.h b/Engine/Plugins/Experimental/Animation/MotionWarping/Source/MotionWarping/Public/AnimNotifyState_MotionWarping.h
--- a/Engine/Plugins/Experimental/Animation/MotionWarping/Source/MotionWarping/Public/AnimNotifyState_MotionWarping.h
+++ b/Engine/Plugins/Experimental/Animation/MotionWarping/Source/MotionWarping/Public/AnimNotifyState_MotionWarping.h
@@ -8,7 +8,7 @@
 
 /** AnimNotifyState used to define a motion warping window in the animation */
 UCLASS(meta = (DisplayName = "Motion Warping"))
-class UAnimNotifyState_MotionWarping : public UAnimNotifyState
+class MOTIONWARPING_API UAnimNotifyState_MotionWarping : public UAnimNotifyState
 {
 	GENERATED_BODY()
 
@@ -18,4 +18,7 @@ public:
 	class URootMotionModifierConfig* RootMotionModifierConfig;
 
 	UAnimNotifyState_MotionWarping(const FObjectInitializer& ObjectInitializer);
+
+	UFUNCTION(BlueprintNativeEvent, Category = "Motion Warping")
+	void AddRootMotionModifier(class UMotionWarpingComponent* MotionWarpingComp, class UAnimSequenceBase* Animation, float StartTime, float EndTime) const;
 };
\ No newline at end of file
diff --git a/Engine/Plugins/Experimental/Animation/MotionWarping/Source/MotionWarping/Public/MotionWarpingComponent.h b/Engine/Plugins/Experimental/Animation/MotionWarping/Source/MotionWarping/Public/MotionWarpingComponent.h
--- a/Engine/Plugins/Experimental/Animation/MotionWarping/Source/MotionWarping/Public/MotionWarpingComponent.h
+++ b/Engine/Plugins/Experimental/Animation/MotionWarping/Source/MotionWarping/Public/MotionWarpingComponent.h
@@ -25,7 +25,7 @@ struct FMotionWarpingCVars
 #endif
 
 UCLASS()
-class UMotionWarpingUtilities : public UBlueprintFunctionLibrary
+class MOTIONWARPING_API UMotionWarpingUtilities : public UBlueprintFunctionLibrary
 {
 	GENERATED_BODY()
 
@@ -49,8 +49,10 @@ public:
 	static FTransform ExtractRootMotionFromAnimation(const UAnimSequenceBase* Animation, float StartTime, float EndTime);
 };
 
+DECLARE_DYNAMIC_MULTICAST_DELEGATE_OneParam(FMotionWarpingPreUpdate, class UMotionWarpingComponent*, MotionWarpingComp);
+
 UCLASS(ClassGroup = Movement, meta = (BlueprintSpawnableComponent))
-class UMotionWarpingComponent : public UActorComponent
+class MOTIONWARPING_API UMotionWarpingComponent : public UActorComponent
 {
 	GENERATED_BODY()
 
@@ -60,6 +62,10 @@ public:
 	UPROPERTY(EditAnywhere, BlueprintReadWrite, Category = "Config")
 	bool bSearchForWindowsInAnimsWithinMontages;
 
+	/** Event called before Root Motion Modifiers are updated */
+	UPROPERTY(BlueprintAssignable, Category = "Motion Warping")
+	FMotionWarpingPreUpdate OnPreUpdate;
+
 	UMotionWarpingComponent(const FObjectInitializer& ObjectInitializer);
 
 	virtual void InitializeComponent() override;
@@ -68,7 +74,7 @@ public:
 	FORCEINLINE ACharacter* GetCharacterOwner() const { return CharacterOwner.Get(); }
 
 	/** Returns the list of root motion modifiers */
-	FORCEINLINE const TArray<TUniquePtr<FRootMotionModifier>>& GetRootMotionModifiers() const { return RootMotionModifiers; }
+	FORCEINLINE const TArray<TSharedPtr<FRootMotionModifier>>& GetRootMotionModifiers() const { return RootMotionModifiers; }
 
 	/** Find the SyncPoint associated with a specified name */
 	FORCEINLINE const FMotionWarpingSyncPoint* FindSyncPoint(const FName& SyncPointName) const { return SyncPoints.Find(SyncPointName); }
@@ -77,8 +83,14 @@ public:
 	UFUNCTION(BlueprintCallable, Category = "Motion Warping")
 	void AddOrUpdateSyncPoint(FName Name, const FMotionWarpingSyncPoint& SyncPoint);
 
-	/** Check if we contain a RootMotionModifier that matches the supplied config data */
-	bool ContainsMatchingModifier(const URootMotionModifierConfig* Config, const UAnimSequenceBase* Animation, float StartTime, float EndTime) const;
+	/** Check if we contain a RootMotionModifier for the supplied animation and time range */
+	bool ContainsModifier(const UAnimSequenceBase* Animation, float StartTime, float EndTime) const;
+
+	/** Add a new modifier */
+	void AddRootMotionModifier(TSharedPtr<FRootMotionModifier> Modifier);
+
+	/** Mark all the modifiers as Disable */
+	void DisableAllRootMotionModifiers();
 
 protected:
 
@@ -87,7 +99,7 @@ protected:
 	TWeakObjectPtr<ACharacter> CharacterOwner;
 
 	/** List of root motion modifiers */
-	TArray<TUniquePtr<FRootMotionModifier>> RootMotionModifiers;
+	TArray<TSharedPtr<FRootMotionModifier>> RootMotionModifiers;
 
 	UPROPERTY(Transient)
 	TMap<FName, FMotionWarpingSyncPoint> SyncPoints;
diff --git a/Engine/Plugins/Experimental/Animation/MotionWarping/Source/MotionWarping/Public/RootMotionModifier.h b/Engine/Plugins/Experimental/Animation/MotionWarping/Source/MotionWarping/Public/RootMotionModifier.h
--- a/Engine/Plugins/Experimental/Animation/MotionWarping/Source/MotionWarping/Public/RootMotionModifier.h
+++ b/Engine/Plugins/Experimental/Animation/MotionWarping/Source/MotionWarping/Public/RootMotionModifier.h
@@ -19,12 +19,15 @@ enum class ERootMotionModifierState : uint8
 	Active,
 
 	/** The modifier has been marked for removal. Usually because the warping window is done */
-	MarkedForRemoval
+	MarkedForRemoval,
+
+	/** The modifier will remain in the list (as long as the window is active) but will not modify the root motion */
+	Disabled
 };
 
 /** Base struct for Root Motion Modifiers */
 USTRUCT()
-struct FRootMotionModifier
+struct MOTIONWARPING_API FRootMotionModifier
 {
 	GENERATED_BODY()
 
@@ -77,28 +80,22 @@ public:
 
 /** Blueprint wrapper around the config properties of a root motion modifier */
 UCLASS(abstract, BlueprintType, EditInlineNew)
-class URootMotionModifierConfig : public UObject
+class MOTIONWARPING_API URootMotionModifierConfig : public UObject
 {
 	GENERATED_BODY()
 
 public:
-
+	
 	URootMotionModifierConfig(const FObjectInitializer& ObjectInitializer)
 		: Super(ObjectInitializer) {}
 
-	/** Creates a RootMotionModifier of the type this object represents */
-	virtual TUniquePtr<FRootMotionModifier> CreateRootMotionModifier(const UAnimSequenceBase* Animation, float StartTime, float EndTime) const PURE_VIRTUAL(URootMotionModifierConfig::CreateRootMotionModifier, return TUniquePtr<FRootMotionModifier>(nullptr););
-
-	/** Checks whether the supplied RootMotionModifier matches the config properties in this object. Should be overridden by subclasses to perform the rest of the comparison. */
-	virtual bool MatchesConfig(const TUniquePtr<FRootMotionModifier>& Modifier, const UAnimSequenceBase* Animation, float StartTime, float EndTime) const
-	{
-		return (Modifier.IsValid() && Modifier->Animation == Animation && Modifier->StartTime == StartTime && Modifier->EndTime == EndTime);
-	}
+	/** Adds a RootMotionModifier of the type this object represents */
+	virtual void AddRootMotionModifier(UMotionWarpingComponent* InMotionWarpingComp, const UAnimSequenceBase* Animation, float StartTime, float EndTime) const PURE_VIRTUAL(URootMotionModifierConfig::AddRootMotionModifier,);
 };
 
 /** Represents a point of alignment in the world */
 USTRUCT(BlueprintType, meta = (HasNativeMake = "MotionWarping.MotionWarpingUtilities.MakeMotionWarpingSyncPoint", HasNativeBreak = "MotionWarping.MotionWarpingUtilities.BreakMotionWarpingSyncPoint"))
-struct FMotionWarpingSyncPoint
+struct MOTIONWARPING_API FMotionWarpingSyncPoint
 {
 	GENERATED_BODY()
 
@@ -107,6 +104,8 @@ struct FMotionWarpingSyncPoint
 		: Location(InLocation), Rotation(InRotation) {}
 	FMotionWarpingSyncPoint(const FVector& InLocation, const FRotator& InRotation)
 		: Location(InLocation), Rotation(InRotation.Quaternion()) {}
+	FMotionWarpingSyncPoint(const FTransform& InTransform)
+		: Location(InTransform.GetLocation()), Rotation(InTransform.GetRotation()) {}
 
 	FORCEINLINE const FVector& GetLocation() const { return Location; }
 	FORCEINLINE const FQuat& GetRotation() const { return Rotation; }
@@ -135,7 +134,7 @@ protected:
 ///////////////////////////////////////////////////////////////
 
 USTRUCT()
-struct FRootMotionModifier_Warp : public FRootMotionModifier
+struct MOTIONWARPING_API FRootMotionModifier_Warp : public FRootMotionModifier
 {
 	GENERATED_BODY()
 
@@ -183,7 +182,7 @@ protected:
 };
 
 UCLASS(meta = (DisplayName = "Simple Warp"))
-class URootMotionModifierConfig_Warp : public URootMotionModifierConfig
+class MOTIONWARPING_API URootMotionModifierConfig_Warp : public URootMotionModifierConfig
 {
 	GENERATED_BODY()
 
@@ -207,41 +206,20 @@ public:
 
 	URootMotionModifierConfig_Warp(const FObjectInitializer& ObjectInitializer) : Super(ObjectInitializer) {}
 
-	virtual TUniquePtr<FRootMotionModifier> CreateRootMotionModifier(const UAnimSequenceBase* Animation, float StartTime, float EndTime) const override
+	virtual void AddRootMotionModifier(UMotionWarpingComponent* MotionWarpingComp, const UAnimSequenceBase* Animation, float StartTime, float EndTime) const override
 	{
-		TUniquePtr<FRootMotionModifier_Warp> Modifier = MakeUnique<FRootMotionModifier_Warp>();
-		Modifier->Animation = Animation;
-		Modifier->StartTime = StartTime;
-		Modifier->EndTime = EndTime;
-		Modifier->SyncPointName = SyncPointName;
-		Modifier->bWarpTranslation = bWarpTranslation;
-		Modifier->bIgnoreZAxis = bIgnoreZAxis;
-		Modifier->bWarpRotation = bWarpRotation;
-		return Modifier;
+		URootMotionModifierConfig_Warp::AddRootMotionModifierSimpleWarp(MotionWarpingComp, Animation, StartTime, EndTime, SyncPointName, bWarpTranslation, bIgnoreZAxis, bWarpRotation);
 	}
 
-	virtual bool MatchesConfig(const TUniquePtr<FRootMotionModifier>& Modifier, const UAnimSequenceBase* Animation, float StartTime, float EndTime) const override
-	{
-		if (Super::MatchesConfig(Modifier, Animation, StartTime, EndTime))
-		{
-			if (Modifier->GetScriptStruct()->IsChildOf(FRootMotionModifier_Warp::StaticStruct()))
-			{
-				const FRootMotionModifier_Warp* ModifierAsWarp = static_cast<const FRootMotionModifier_Warp*>(Modifier.Get());
-				return (ModifierAsWarp->SyncPointName == SyncPointName &&
-					ModifierAsWarp->bWarpTranslation == bWarpTranslation &&
-					ModifierAsWarp->bIgnoreZAxis == bIgnoreZAxis &&
-					ModifierAsWarp->bWarpRotation == bWarpRotation);
-			}
-		}
-		return false;
-	}
+	UFUNCTION(BlueprintCallable, Category = "Motion Warping")
+	static void AddRootMotionModifierSimpleWarp(UMotionWarpingComponent* InMotionWarpingComp, const UAnimSequenceBase* InAnimation, float InStartTime, float InEndTime, FName InSyncPointName, bool bInWarpTranslation, bool bInIgnoreZAxis, bool bInWarpRotation);
 };
 
 // FRootMotionModifier_Scale
 ///////////////////////////////////////////////////////////////
 
 USTRUCT()
-struct FRootMotionModifier_Scale : public FRootMotionModifier
+struct MOTIONWARPING_API FRootMotionModifier_Scale : public FRootMotionModifier
 {
 	GENERATED_BODY()
 
@@ -264,7 +242,7 @@ public:
 };
 
 UCLASS(meta = (DisplayName = "Scale"))
-class URootMotionModifierConfig_Scale : public URootMotionModifierConfig
+class MOTIONWARPING_API URootMotionModifierConfig_Scale : public URootMotionModifierConfig
 {
 	GENERATED_BODY()
 
@@ -276,26 +254,11 @@ public:
 
 	URootMotionModifierConfig_Scale(const FObjectInitializer& ObjectInitializer) : Super(ObjectInitializer) {}
 
-	virtual TUniquePtr<FRootMotionModifier> CreateRootMotionModifier(const UAnimSequenceBase* Animation, float StartTime, float EndTime) const override
+	virtual void AddRootMotionModifier(UMotionWarpingComponent* MotionWarpingComp, const UAnimSequenceBase* Animation, float StartTime, float EndTime) const override
 	{
-		TUniquePtr<FRootMotionModifier_Scale> Modifier = MakeUnique<FRootMotionModifier_Scale>();
-		Modifier->Animation = Animation;
-		Modifier->StartTime = StartTime;
-		Modifier->EndTime = EndTime;
-		Modifier->Scale = Scale;
-		return Modifier;
+		URootMotionModifierConfig_Scale::AddRootMotionModifierScale(MotionWarpingComp, Animation, StartTime, EndTime, Scale);
 	}
 
-	virtual bool MatchesConfig(const TUniquePtr<FRootMotionModifier>& Modifier, const UAnimSequenceBase* Animation, float StartTime, float EndTime) const override
-	{
-		if (Super::MatchesConfig(Modifier, Animation, StartTime, EndTime))
-		{
-			if (Modifier->GetScriptStruct()->IsChildOf(FRootMotionModifier_Scale::StaticStruct()))
-			{
-				const FRootMotionModifier_Scale* ModifierAsSimple = static_cast<const FRootMotionModifier_Scale*>(Modifier.Get());
-				return (ModifierAsSimple->Scale == Scale);
-			}
-		}
-		return false;
-	}
+	UFUNCTION(BlueprintCallable, Category = "Motion Warping")
+	static void AddRootMotionModifierScale(UMotionWarpingComponent* InMotionWarpingComp, const UAnimSequenceBase* InAnimation, float InStartTime, float InEndTime, FVector InScale);
 };
\ No newline at end of file
diff --git a/Engine/Plugins/Experimental/Animation/MotionWarping/Source/MotionWarping/Public/RootMotionModifier_AdjustmentBlendWarp.h b/Engine/Plugins/Experimental/Animation/MotionWarping/Source/MotionWarping/Public/RootMotionModifier_AdjustmentBlendWarp.h
--- a/Engine/Plugins/Experimental/Animation/MotionWarping/Source/MotionWarping/Public/RootMotionModifier_AdjustmentBlendWarp.h
+++ b/Engine/Plugins/Experimental/Animation/MotionWarping/Source/MotionWarping/Public/RootMotionModifier_AdjustmentBlendWarp.h
@@ -42,7 +42,7 @@ struct FMotionDeltaTrackContainer
 };
 
 USTRUCT()
-struct FRootMotionModifier_AdjustmentBlendWarp : public FRootMotionModifier_Warp
+struct MOTIONWARPING_API FRootMotionModifier_AdjustmentBlendWarp : public FRootMotionModifier_Warp
 {
 	GENERATED_BODY()
 
@@ -98,7 +98,7 @@ protected:
 };
 
 UCLASS(meta = (DisplayName = "Adjustment Blend Warp"))
-class URootMotionModifierConfig_AdjustmentBlendWarp : public URootMotionModifierConfig_Warp
+class MOTIONWARPING_API URootMotionModifierConfig_AdjustmentBlendWarp : public URootMotionModifierConfig_Warp
 {
 	GENERATED_BODY()
 
@@ -113,21 +113,14 @@ public:
 	URootMotionModifierConfig_AdjustmentBlendWarp(const FObjectInitializer& ObjectInitializer)
 		: Super(ObjectInitializer) {}
 
-	virtual TUniquePtr<FRootMotionModifier> CreateRootMotionModifier(const UAnimSequenceBase* Animation, float StartTime, float EndTime) const override
+	virtual void AddRootMotionModifier(UMotionWarpingComponent* MotionWarpingComp, const UAnimSequenceBase* Animation, float StartTime, float EndTime) const override
 	{
-		TUniquePtr<FRootMotionModifier_AdjustmentBlendWarp> Modifier = MakeUnique<FRootMotionModifier_AdjustmentBlendWarp>();
-		Modifier->Animation = Animation;
-		Modifier->StartTime = StartTime;
-		Modifier->EndTime = EndTime;
-		Modifier->SyncPointName = SyncPointName;
-		Modifier->bWarpTranslation = bWarpTranslation;
-		Modifier->bIgnoreZAxis = bIgnoreZAxis;
-		Modifier->bWarpRotation = bWarpRotation;
-		Modifier->bWarpIKBones = bWarpIKBones;
-		Modifier->IKBones = IKBones;
-		return Modifier;
+		URootMotionModifierConfig_AdjustmentBlendWarp::AddRootMotionModifierAdjustmentBlendWarp(MotionWarpingComp, Animation, StartTime, EndTime, SyncPointName, bWarpTranslation, bIgnoreZAxis, bWarpRotation, bWarpIKBones, IKBones);
 	}
 
+	UFUNCTION(BlueprintCallable, Category = "Motion Warping")
+	static void AddRootMotionModifierAdjustmentBlendWarp(UMotionWarpingComponent* InMotionWarpingComp, const UAnimSequenceBase* InAnimation, float InStartTime, float InEndTime, FName InSyncPointName, bool bInWarpTranslation, bool bInIgnoreZAxis, bool bInWarpRotation, bool bInWarpIKBones, const TArray<FName>& InIKBones);
+
 	UFUNCTION(BlueprintPure, Category = "Motion Warping")
 	static void GetIKBoneTransformAndAlpha(ACharacter* Character, FName BoneName, FTransform& OutTransform, float& OutAlpha);
 };
\ No newline at end of file
diff --git a/Engine/Plugins/Experimental/Animation/MotionWarping/Source/MotionWarping/Public/RootMotionModifier_SkewWarp.h b/Engine/Plugins/Experimental/Animation/MotionWarping/Source/MotionWarping/Public/RootMotionModifier_SkewWarp.h
--- a/Engine/Plugins/Experimental/Animation/MotionWarping/Source/MotionWarping/Public/RootMotionModifier_SkewWarp.h
+++ b/Engine/Plugins/Experimental/Animation/MotionWarping/Source/MotionWarping/Public/RootMotionModifier_SkewWarp.h
@@ -7,7 +7,7 @@
 #include "RootMotionModifier_SkewWarp.generated.h"
 
 USTRUCT()
-struct FRootMotionModifier_SkewWarp : public FRootMotionModifier_Warp
+struct MOTIONWARPING_API FRootMotionModifier_SkewWarp : public FRootMotionModifier_Warp
 {
 	GENERATED_BODY()
 
@@ -21,7 +21,7 @@ public:
 };
 
 UCLASS(meta = (DisplayName = "Skew Warp"))
-class URootMotionModifierConfig_SkewWarp : public URootMotionModifierConfig_Warp
+class MOTIONWARPING_API URootMotionModifierConfig_SkewWarp : public URootMotionModifierConfig_Warp
 {
 	GENERATED_BODY()
 
@@ -30,16 +30,11 @@ public:
 	URootMotionModifierConfig_SkewWarp(const FObjectInitializer& ObjectInitializer)
 		: Super(ObjectInitializer) {}
 
-	virtual TUniquePtr<FRootMotionModifier> CreateRootMotionModifier(const UAnimSequenceBase* Animation, float StartTime, float EndTime) const override
+	virtual void AddRootMotionModifier(UMotionWarpingComponent* MotionWarpingComp, const UAnimSequenceBase* Animation, float StartTime, float EndTime) const override
 	{
-		TUniquePtr<FRootMotionModifier_SkewWarp> Modifier = MakeUnique<FRootMotionModifier_SkewWarp>();
-		Modifier->Animation = Animation;
-		Modifier->StartTime = StartTime;
-		Modifier->EndTime = EndTime;
-		Modifier->SyncPointName = SyncPointName;
-		Modifier->bWarpTranslation = bWarpTranslation;
-		Modifier->bIgnoreZAxis = bIgnoreZAxis;
-		Modifier->bWarpRotation = bWarpRotation;
-		return Modifier;
+		URootMotionModifierConfig_SkewWarp::AddRootMotionModifierSkewWarp(MotionWarpingComp, Animation, StartTime, EndTime, SyncPointName, bWarpTranslation, bIgnoreZAxis, bWarpRotation);
 	}
+
+	UFUNCTION(BlueprintCallable, Category = "Motion Warping")
+	static void AddRootMotionModifierSkewWarp(UMotionWarpingComponent* InMotionWarpingComp, const UAnimSequenceBase* InAnimation, float InStartTime, float InEndTime, FName InSyncPointName, bool bInWarpTranslation, bool bInIgnoreZAxis, bool bInWarpRotation);
 };
\ No newline at end of file
