Motion Warping - Fixed bug where warping window for animations withing a montage were ignored. Fixed bug where rotation warping wasn't working. Added function to remove sync points. #review trivial #fyi [REDACTED_HANDLE]

diff --git a/Engine/Plugins/Experimental/Animation/MotionWarping/Source/MotionWarping/Private/MotionWarpingComponent.cpp b/Engine/Plugins/Experimental/Animation/MotionWarping/Source/MotionWarping/Private/MotionWarpingComponent.cpp
--- a/Engine/Plugins/Experimental/Animation/MotionWarping/Source/MotionWarping/Private/MotionWarpingComponent.cpp
+++ b/Engine/Plugins/Experimental/Animation/MotionWarping/Source/MotionWarping/Private/MotionWarpingComponent.cpp
@@ -174,8 +174,7 @@ void UMotionWarpingComponent::Update()
 					for (const FAnimNotifyEvent& NotifyEvent : AnimSegment->AnimReference->Notifies)
 					{
 						const UAnimNotifyState_MotionWarping* MotionWarpingNotify = NotifyEvent.NotifyStateClass ? Cast<UAnimNotifyState_MotionWarping>(NotifyEvent.NotifyStateClass) : nullptr;
-						const URootMotionModifierConfig* MotionWarpingConfig = MotionWarpingNotify ? MotionWarpingNotify->RootMotionModifierConfig : nullptr;
-						if (MotionWarpingConfig)
+						if (MotionWarpingNotify)
 						{
 							const float NotifyStartTime = FMath::Clamp(NotifyEvent.GetTriggerTime(), 0.f, AnimSegment->AnimReference->GetPlayLength());
 							const float NotifyEndTime = FMath::Clamp(NotifyEvent.GetEndTriggerTime(), 0.f, AnimSegment->AnimReference->GetPlayLength());
@@ -271,4 +270,9 @@ void UMotionWarpingComponent::AddOrUpdateSyncPoint(FName Name, const FMotionWarp
 		FMotionWarpingSyncPoint& MotionWarpingSyncPoint = SyncPoints.FindOrAdd(Name);
 		MotionWarpingSyncPoint = SyncPoint;
 	}
+}
+
+int32 UMotionWarpingComponent::RemoveSyncPoint(FName Name)
+{
+	return SyncPoints.Remove(Name);
 }
\ No newline at end of file
diff --git a/Engine/Plugins/Experimental/Animation/MotionWarping/Source/MotionWarping/Private/RootMotionModifier.cpp b/Engine/Plugins/Experimental/Animation/MotionWarping/Source/MotionWarping/Private/RootMotionModifier.cpp
--- a/Engine/Plugins/Experimental/Animation/MotionWarping/Source/MotionWarping/Private/RootMotionModifier.cpp
+++ b/Engine/Plugins/Experimental/Animation/MotionWarping/Source/MotionWarping/Private/RootMotionModifier.cpp
@@ -124,15 +124,22 @@ FTransform FRootMotionModifier_Warp::ProcessRootMotion(UMotionWarpingComponent&
 
 	if (bWarpRotation)
 	{
-		const FQuat WarpedRotation = WarpRotation(OwnerComp, RootMotionTotal, InRootMotion, DeltaSeconds);
+		const FQuat WarpedRotation = WarpRotation(OwnerComp, InRootMotion, RootMotionTotal, DeltaSeconds);
 		FinalRootMotion.SetRotation(WarpedRotation);
 	}
 
 	// Debug
 #if !(UE_BUILD_SHIPPING || UE_BUILD_TEST)
-	if (FMotionWarpingCVars::CVarMotionWarpingDebug.GetValueOnGameThread() > 0)
+	const int32 DebugLevel = FMotionWarpingCVars::CVarMotionWarpingDebug.GetValueOnGameThread();
+	if (DebugLevel > 0)
 	{
 		PrintLog(OwnerComp, TEXT("FRootMotionModifier_Simple"), InRootMotion, FinalRootMotion);
+
+		if (DebugLevel >= 2)
+		{
+			const float DrawDebugDuration = FMotionWarpingCVars::CVarMotionWarpingDrawDebugDuration.GetValueOnGameThread();
+			DrawDebugCoordinateSystem(OwnerComp.GetWorld(), CachedSyncPoint.GetLocation(), CachedSyncPoint.Rotator(), 50.f, false, DrawDebugDuration, 0, 1.f);
+		}
 	}
 #endif
 
diff --git a/Engine/Plugins/Experimental/Animation/MotionWarping/Source/MotionWarping/Private/RootMotionModifier_SkewWarp.cpp b/Engine/Plugins/Experimental/Animation/MotionWarping/Source/MotionWarping/Private/RootMotionModifier_SkewWarp.cpp
--- a/Engine/Plugins/Experimental/Animation/MotionWarping/Source/MotionWarping/Private/RootMotionModifier_SkewWarp.cpp
+++ b/Engine/Plugins/Experimental/Animation/MotionWarping/Source/MotionWarping/Private/RootMotionModifier_SkewWarp.cpp
@@ -138,15 +138,22 @@ FTransform FRootMotionModifier_SkewWarp::ProcessRootMotion(UMotionWarpingCompone
 
 	if(bWarpRotation)
 	{
-		const FQuat WarpedRotation = WarpRotation(OwnerComp, RootMotionTotal, InRootMotion, DeltaSeconds);
+		const FQuat WarpedRotation = WarpRotation(OwnerComp, InRootMotion, RootMotionTotal, DeltaSeconds);
 		FinalRootMotion.SetRotation(WarpedRotation);
 	}
 
 	// Debug
 #if !(UE_BUILD_SHIPPING || UE_BUILD_TEST)
-	if (FMotionWarpingCVars::CVarMotionWarpingDebug.GetValueOnGameThread() > 0)
+	const int32 DebugLevel = FMotionWarpingCVars::CVarMotionWarpingDebug.GetValueOnGameThread();
+	if (DebugLevel > 0)
 	{
 		PrintLog(OwnerComp, TEXT("FRootMotionModifier_Skew"), InRootMotion, FinalRootMotion);
+
+		if (DebugLevel >= 2)
+		{
+			const float DrawDebugDuration = FMotionWarpingCVars::CVarMotionWarpingDrawDebugDuration.GetValueOnGameThread();
+			DrawDebugCoordinateSystem(OwnerComp.GetWorld(), CachedSyncPoint.GetLocation(), CachedSyncPoint.Rotator(), 50.f, false, DrawDebugDuration, 0, 1.f);
+		}
 	}
 #endif
 
diff --git a/Engine/Plugins/Experimental/Animation/MotionWarping/Source/MotionWarping/Public/MotionWarpingComponent.h b/Engine/Plugins/Experimental/Animation/MotionWarping/Source/MotionWarping/Public/MotionWarpingComponent.h
--- a/Engine/Plugins/Experimental/Animation/MotionWarping/Source/MotionWarping/Public/MotionWarpingComponent.h
+++ b/Engine/Plugins/Experimental/Animation/MotionWarping/Source/MotionWarping/Public/MotionWarpingComponent.h
@@ -83,6 +83,10 @@ public:
 	UFUNCTION(BlueprintCallable, Category = "Motion Warping")
 	void AddOrUpdateSyncPoint(FName Name, const FMotionWarpingSyncPoint& SyncPoint);
 
+	/** Removes sync point associated with the specified key  */
+	UFUNCTION(BlueprintCallable, Category = "Motion Warping")
+	int32 RemoveSyncPoint(FName Name);
+
 	/** Check if we contain a RootMotionModifier for the supplied animation and time range */
 	bool ContainsModifier(const UAnimSequenceBase* Animation, float StartTime, float EndTime) const;
 
