Motion Warping - Motion modifiers are now marked as disabled instead of for removal when the sync point is not valid. Added Extra logs with more data to help debugging and exposed helper function to extract rootmotion. #review-15152206 trivial #fyi [REDACTED_HANDLE]

diff --git a/Engine/Plugins/Experimental/Animation/MotionWarping/Source/MotionWarping/Private/MotionWarpingComponent.cpp b/Engine/Plugins/Experimental/Animation/MotionWarping/Source/MotionWarping/Private/MotionWarpingComponent.cpp
--- a/Engine/Plugins/Experimental/Animation/MotionWarping/Source/MotionWarping/Private/MotionWarpingComponent.cpp
+++ b/Engine/Plugins/Experimental/Animation/MotionWarping/Source/MotionWarping/Private/MotionWarpingComponent.cpp
@@ -115,8 +115,8 @@ void UMotionWarpingComponent::AddRootMotionModifier(TSharedPtr<FRootMotionModifi
 	{
 		RootMotionModifiers.Add(Modifier);
 
-		UE_LOG(LogMotionWarping, Verbose, TEXT("UMotionWarpingComponent RootMotionModifier added. NetMode: %d WorldTime: %f Animation: %s [%f %f] [%f %f] Loc: %s Rot: %s"),
-			GetWorld()->GetNetMode(), GetWorld()->GetTimeSeconds(), *GetNameSafe(Modifier->Animation.Get()), Modifier->StartTime, Modifier->EndTime, Modifier->PreviousPosition, Modifier->CurrentPosition,
+		UE_LOG(LogMotionWarping, Verbose, TEXT("MotionWarping: RootMotionModifier added. NetMode: %d WorldTime: %f Char: %s Animation: %s [%f %f] [%f %f] Loc: %s Rot: %s"),
+			GetWorld()->GetNetMode(), GetWorld()->GetTimeSeconds(), *GetNameSafe(GetCharacterOwner()), *GetNameSafe(Modifier->Animation.Get()), Modifier->StartTime, Modifier->EndTime, Modifier->PreviousPosition, Modifier->CurrentPosition,
 			*GetCharacterOwner()->GetActorLocation().ToString(), *GetCharacterOwner()->GetActorRotation().ToCompactString());
 	}
 }
@@ -208,9 +208,19 @@ void UMotionWarpingComponent::Update()
 		}
 
 		// Remove the modifiers that has been marked for removal
-		const int32 TotalRemoved = RootMotionModifiers.RemoveAll([](const TSharedPtr<FRootMotionModifier>& Modifier) { return Modifier->State == ERootMotionModifierState::MarkedForRemoval; });
+		RootMotionModifiers.RemoveAll([this](const TSharedPtr<FRootMotionModifier>& Modifier) 
+		{ 
+			if(Modifier->State == ERootMotionModifierState::MarkedForRemoval)
+			{
+				UE_LOG(LogMotionWarping, Verbose, TEXT("MotionWarping: RootMotionModifier removed. NetMode: %d WorldTime: %f Char: %s Animation: %s [%f %f] [%f %f] Loc: %s Rot: %s"),
+					GetWorld()->GetNetMode(), GetWorld()->GetTimeSeconds(), *GetNameSafe(GetCharacterOwner()), *GetNameSafe(Modifier->Animation.Get()), Modifier->StartTime, Modifier->EndTime, Modifier->PreviousPosition, Modifier->CurrentPosition,
+					*GetCharacterOwner()->GetActorLocation().ToString(), *GetCharacterOwner()->GetActorRotation().ToCompactString());
+
+				return true;
+			}
 
-		UE_CLOG(TotalRemoved > 0, LogMotionWarping, Verbose, TEXT("UMotionWarpingComponent::UpdateModifiers Modifiers removed after update: %d"), TotalRemoved);
+			return false; 
+		});
 	}
 }
 
diff --git a/Engine/Plugins/Experimental/Animation/MotionWarping/Source/MotionWarping/Private/RootMotionModifier.cpp b/Engine/Plugins/Experimental/Animation/MotionWarping/Source/MotionWarping/Private/RootMotionModifier.cpp
--- a/Engine/Plugins/Experimental/Animation/MotionWarping/Source/MotionWarping/Private/RootMotionModifier.cpp
+++ b/Engine/Plugins/Experimental/Animation/MotionWarping/Source/MotionWarping/Private/RootMotionModifier.cpp
@@ -19,6 +19,9 @@ void FRootMotionModifier::Update(UMotionWarpingComponent& OwnerComp)
 	// Mark for removal if our animation is not relevant anymore
 	if (Montage == nullptr || Montage != Animation)
 	{
+		UE_LOG(LogMotionWarping, Verbose, TEXT("MotionWarping: Marking RootMotionModifier for removal. Reason: Animation is not valid. Char: %s Current Montage: %s. Window: Animation: %s [%f %f] [%f %f]"),
+			*GetNameSafe(OwnerComp.GetCharacterOwner()), *GetNameSafe(Montage), *GetNameSafe(Animation.Get()), StartTime, EndTime, PreviousPosition, CurrentPosition);
+
 		State = ERootMotionModifierState::MarkedForRemoval;
 		return;
 	}
@@ -31,6 +34,9 @@ void FRootMotionModifier::Update(UMotionWarpingComponent& OwnerComp)
 	// Mark for removal if the animation already passed the warping window
 	if (PreviousPosition >= EndTime)
 	{
+		UE_LOG(LogMotionWarping, Verbose, TEXT("MotionWarping: Marking RootMotionModifier for removal. Reason: Window has ended. Char: %s Animation: %s [%f %f] [%f %f]"),
+			*GetNameSafe(OwnerComp.GetCharacterOwner()), *GetNameSafe(Animation.Get()), StartTime, EndTime, PreviousPosition, CurrentPosition);
+
 		State = ERootMotionModifierState::MarkedForRemoval;
 		return;
 	}
@@ -51,21 +57,24 @@ void FRootMotionModifier::Update(UMotionWarpingComponent& OwnerComp)
 
 void FRootMotionModifier_Warp::Update(UMotionWarpingComponent& OwnerComp)
 {
-	// Mark for removal if there is no sync point for us
-	const FMotionWarpingSyncPoint* SyncPointPtr = OwnerComp.FindSyncPoint(SyncPointName);
-	if (SyncPointPtr == nullptr)
-	{
-		UE_LOG(LogMotionWarping, Warning, TEXT("Marking RootMotionModifier for removal. Reason: Can't find SyncPoint (%s)"), *SyncPointName.ToString());
-		State = ERootMotionModifierState::MarkedForRemoval;
-		return;
-	}
-
 	// Update playback times and state
 	FRootMotionModifier::Update(OwnerComp);
 
 	// Cache sync point transform and trigger OnSyncPointChanged if needed
 	if (State == ERootMotionModifierState::Active)
 	{
+		const FMotionWarpingSyncPoint* SyncPointPtr = OwnerComp.FindSyncPoint(SyncPointName);
+
+		// Disable if there is no sync point for us
+		if (SyncPointPtr == nullptr)
+		{
+			UE_LOG(LogMotionWarping, Verbose, TEXT("MotionWarping: Marking RootMotionModifier as Disabled. Reason: Invalid Sync Point (%s). Char: %s Animation: %s [%f %f] [%f %f]"),
+				*SyncPointName.ToString(), *GetNameSafe(OwnerComp.GetCharacterOwner()), *GetNameSafe(Animation.Get()), StartTime, EndTime, PreviousPosition, CurrentPosition);
+
+			State = ERootMotionModifierState::Disabled;
+			return;
+		}
+
 		if (CachedSyncPoint != *SyncPointPtr)
 		{
 			CachedSyncPoint = *SyncPointPtr;
@@ -181,7 +190,7 @@ void FRootMotionModifier_Warp::PrintLog(const UMotionWarpingComponent& OwnerComp
 	const float Speed = WarpedRootMotion.GetTranslation().Size() / DeltaSeconds;
 	const float EndTimeOffset = CurrentPosition - EndTime;
 
-	UE_LOG(LogMotionWarping, Log, TEXT("%s. NetMode: %d Char: %s Anim: %s Window [%f %f][%f %f] DeltaTime: %f WorldTime: %f EndTimeOffset: %f Dist2D: %f FutureDist2D: %f Dot: %f OriginalMotionDelta: %s (%f) FinalMotionDelta: %s (%f) Speed: %f Location: %s FutureLocation: %s Rotation: %s FutureRotation: %s"),
+	UE_LOG(LogMotionWarping, Log, TEXT("MotionWarping: %s. NetMode: %d Char: %s Anim: %s Window [%f %f][%f %f] DeltaTime: %f WorldTime: %f EndTimeOffset: %f Dist2D: %f FutureDist2D: %f Dot: %f OriginalMotionDelta: %s (%f) FinalMotionDelta: %s (%f) Speed: %f Loc: %s FutureLoc: %s Rot: %s FutureRot: %s"),
 		*Name, (int32)CharacterOwner->GetWorld()->GetNetMode(), *GetNameSafe(CharacterOwner), *GetNameSafe(Animation.Get()), StartTime, EndTime, PreviousPosition, CurrentPosition, DeltaSeconds, CharacterOwner->GetWorld()->GetTimeSeconds(), EndTimeOffset, CurrentDist2D, FutureDist2D, Dot,
 		*OriginalRootMotion.GetTranslation().ToString(), OriginalRootMotion.GetTranslation().Size(), *WarpedRootMotion.GetTranslation().ToString(), WarpedRootMotion.GetTranslation().Size(), Speed,
 		*CurrentLocation.ToString(), *FutureLocation.ToString(), *CurrentRotation.ToCompactString(), *FutureRotation.ToCompactString());
diff --git a/Engine/Plugins/Experimental/Animation/MotionWarping/Source/MotionWarping/Public/MotionWarpingComponent.h b/Engine/Plugins/Experimental/Animation/MotionWarping/Source/MotionWarping/Public/MotionWarpingComponent.h
--- a/Engine/Plugins/Experimental/Animation/MotionWarping/Source/MotionWarping/Public/MotionWarpingComponent.h
+++ b/Engine/Plugins/Experimental/Animation/MotionWarping/Source/MotionWarping/Public/MotionWarpingComponent.h
@@ -46,6 +46,7 @@ public:
 	static void ExtractComponentSpacePose(const UAnimSequenceBase* Animation, const FBoneContainer& BoneContainer, float Time, bool bExtractRootMotion, FCSPose<FCompactPose>& OutPose);
 
 	/** Extract Root Motion transform from a contiguous position range */
+	UFUNCTION(BlueprintCallable, Category = "Motion Warping")
 	static FTransform ExtractRootMotionFromAnimation(const UAnimSequenceBase* Animation, float StartTime, float EndTime);
 };
 
