Motion Warping - Fix bug in adjustment blend warp preventing the actor from reaching the exact target. #review-15152356 trivial #fyi [REDACTED_HANDLE]

diff --git a/Engine/Plugins/Experimental/Animation/MotionWarping/Source/MotionWarping/Private/RootMotionModifier_AdjustmentBlendWarp.cpp b/Engine/Plugins/Experimental/Animation/MotionWarping/Source/MotionWarping/Private/RootMotionModifier_AdjustmentBlendWarp.cpp
--- a/Engine/Plugins/Experimental/Animation/MotionWarping/Source/MotionWarping/Private/RootMotionModifier_AdjustmentBlendWarp.cpp
+++ b/Engine/Plugins/Experimental/Animation/MotionWarping/Source/MotionWarping/Private/RootMotionModifier_AdjustmentBlendWarp.cpp
@@ -57,7 +57,7 @@ void FRootMotionModifier_AdjustmentBlendWarp::OnSyncPointChanged(UMotionWarpingC
 	const USkeletalMeshComponent* SkelMeshComp = CharacterOwner->GetMesh();
 
 	ActualStartTime = PreviousPosition;
-	CachedRootMotion = UMotionWarpingUtilities::ExtractRootMotionFromAnimation(Animation.Get(), StartTime, EndTime);
+	CachedRootMotion = UMotionWarpingUtilities::ExtractRootMotionFromAnimation(Animation.Get(), ActualStartTime, EndTime);
 	CachedMeshRelativeTransform = SkelMeshComp->GetRelativeTransform();
 	CachedMeshTransform = SkelMeshComp->GetComponentTransform();
 
