Motion Warping -  Added option to define whether rotation should be warp to match the rotation of the sync point or to face the sync point.  Added optional multiplier to modify how fast the rotation is warped relative to the duration of the window. #review-15271229 [REDACTED_HANDLE] #fyi [REDACTED_HANDLE]

diff --git a/Engine/Plugins/Experimental/Animation/MotionWarping/Source/MotionWarping/Private/RootMotionModifier.cpp b/Engine/Plugins/Experimental/Animation/MotionWarping/Source/MotionWarping/Private/RootMotionModifier.cpp
--- a/Engine/Plugins/Experimental/Animation/MotionWarping/Source/MotionWarping/Private/RootMotionModifier.cpp
+++ b/Engine/Plugins/Experimental/Animation/MotionWarping/Source/MotionWarping/Private/RootMotionModifier.cpp
@@ -155,18 +155,33 @@ FTransform FRootMotionModifier_Warp::ProcessRootMotion(UMotionWarpingComponent&
 	return FinalRootMotion;
 }
 
-FQuat FRootMotionModifier_Warp::WarpRotation(UMotionWarpingComponent& OwnerComp, const FTransform& RootMotionDelta, const FTransform& RootMotionTotal, float DeltaSeconds)
+FQuat FRootMotionModifier_Warp::GetTargetRotation(UMotionWarpingComponent& OwnerComp) const
 {
 	const FTransform& CharacterTransform = OwnerComp.GetCharacterOwner()->GetActorTransform();
 
-	const FQuat CurrentRotation = CharacterTransform.GetRotation();
-	const float TimeRemaining = EndTime - PreviousPosition;
-	const FQuat Desired = CachedSyncPoint.GetRotation();
+	if (RotationType == EMotionWarpRotationType::Default)
+	{
+		return CachedSyncPoint.GetRotation();
+	}
+	else if (RotationType == EMotionWarpRotationType::Facing)
+	{
+		const FVector ToSyncPoint = (CachedSyncPoint.GetLocation() - CharacterTransform.GetLocation()).GetSafeNormal2D();
+		return FRotationMatrix::MakeFromXZ(ToSyncPoint, FVector::UpVector).ToQuat();
+	}
 
+	return FQuat::Identity;
+}
+
+FQuat FRootMotionModifier_Warp::WarpRotation(UMotionWarpingComponent& OwnerComp, const FTransform& RootMotionDelta, const FTransform& RootMotionTotal, float DeltaSeconds)
+{
+	const FTransform& CharacterTransform = OwnerComp.GetCharacterOwner()->GetActorTransform();
+	const FQuat CurrentRotation = CharacterTransform.GetRotation();
+	const FQuat TargetRotation = GetTargetRotation(OwnerComp);
+	const float TimeRemaining = (EndTime - PreviousPosition) * WarpRotationTimeMultiplier;
 	const FQuat RemainingRootRotationInWorld = RootMotionTotal.GetRotation();
 	const FQuat CurrentPlusRemainingRootMotion = RemainingRootRotationInWorld * CurrentRotation;
 	const float PercentThisStep = FMath::Clamp(DeltaSeconds / TimeRemaining, 0.f, 1.f);
-	const FQuat TargetRotThisFrame = FQuat::Slerp(CurrentPlusRemainingRootMotion, Desired, PercentThisStep);
+	const FQuat TargetRotThisFrame = FQuat::Slerp(CurrentPlusRemainingRootMotion, TargetRotation, PercentThisStep);
 	const FQuat DeltaOut = TargetRotThisFrame * CurrentPlusRemainingRootMotion.Inverse();
 
 	return (DeltaOut * RootMotionDelta.GetRotation());
@@ -200,7 +215,7 @@ void FRootMotionModifier_Warp::PrintLog(const UMotionWarpingComponent& OwnerComp
 // URootMotionModifierConfig_Warp
 ///////////////////////////////////////////////////////////////
 
-void URootMotionModifierConfig_Warp::AddRootMotionModifierSimpleWarp(UMotionWarpingComponent* InMotionWarpingComp, const UAnimSequenceBase* InAnimation, float InStartTime, float InEndTime, FName InSyncPointName, bool bInWarpTranslation, bool bInIgnoreZAxis, bool bInWarpRotation)
+void URootMotionModifierConfig_Warp::AddRootMotionModifierSimpleWarp(UMotionWarpingComponent* InMotionWarpingComp, const UAnimSequenceBase* InAnimation, float InStartTime, float InEndTime, FName InSyncPointName, bool bInWarpTranslation, bool bInIgnoreZAxis, bool bInWarpRotation, EMotionWarpRotationType InRotationType, float InWarpRotationTimeMultiplier)
 {
 	if (ensureAlways(InMotionWarpingComp))
 	{
@@ -212,6 +227,8 @@ void URootMotionModifierConfig_Warp::AddRootMotionModifierSimpleWarp(UMotionWarp
 		NewModifier->bWarpTranslation = bInWarpTranslation;
 		NewModifier->bIgnoreZAxis = bInIgnoreZAxis;
 		NewModifier->bWarpRotation = bInWarpRotation;
+		NewModifier->RotationType = InRotationType;
+		NewModifier->WarpRotationTimeMultiplier = InWarpRotationTimeMultiplier;
 		InMotionWarpingComp->AddRootMotionModifier(NewModifier);
 	}
 }
diff --git a/Engine/Plugins/Experimental/Animation/MotionWarping/Source/MotionWarping/Private/RootMotionModifier_AdjustmentBlendWarp.cpp b/Engine/Plugins/Experimental/Animation/MotionWarping/Source/MotionWarping/Private/RootMotionModifier_AdjustmentBlendWarp.cpp
--- a/Engine/Plugins/Experimental/Animation/MotionWarping/Source/MotionWarping/Private/RootMotionModifier_AdjustmentBlendWarp.cpp
+++ b/Engine/Plugins/Experimental/Animation/MotionWarping/Source/MotionWarping/Private/RootMotionModifier_AdjustmentBlendWarp.cpp
@@ -148,7 +148,7 @@ void FRootMotionModifier_AdjustmentBlendWarp::PrecomputeWarpedTracks(UMotionWarp
 	FQuat RootTotalAdditiveRotation = FQuat::Identity;
 	if (bWarpRotation)
 	{
-		const FQuat TargetRotation = CachedSyncPoint.GetRotation();
+		const FQuat TargetRotation = GetTargetRotation(OwnerComp);
 		const FQuat OriginalRotation = CachedMeshRelativeTransform.GetRotation().Inverse() * (CachedRootMotion * CachedMeshTransform).GetRotation();
 		RootTotalAdditiveRotation = FQuat::FindBetweenNormals(OriginalRotation.GetForwardVector(), TargetRotation.GetForwardVector());
 	}
diff --git a/Engine/Plugins/Experimental/Animation/MotionWarping/Source/MotionWarping/Private/RootMotionModifier_SkewWarp.cpp b/Engine/Plugins/Experimental/Animation/MotionWarping/Source/MotionWarping/Private/RootMotionModifier_SkewWarp.cpp
--- a/Engine/Plugins/Experimental/Animation/MotionWarping/Source/MotionWarping/Private/RootMotionModifier_SkewWarp.cpp
+++ b/Engine/Plugins/Experimental/Animation/MotionWarping/Source/MotionWarping/Private/RootMotionModifier_SkewWarp.cpp
@@ -163,7 +163,7 @@ FTransform FRootMotionModifier_SkewWarp::ProcessRootMotion(UMotionWarpingCompone
 // URootMotionModifierConfig_SkewWarp
 ///////////////////////////////////////////////////////////////
 
-void URootMotionModifierConfig_SkewWarp::AddRootMotionModifierSkewWarp(UMotionWarpingComponent* InMotionWarpingComp, const UAnimSequenceBase* InAnimation, float InStartTime, float InEndTime, FName InSyncPointName, bool bInWarpTranslation, bool bInIgnoreZAxis, bool bInWarpRotation)
+void URootMotionModifierConfig_SkewWarp::AddRootMotionModifierSkewWarp(UMotionWarpingComponent* InMotionWarpingComp, const UAnimSequenceBase* InAnimation, float InStartTime, float InEndTime, FName InSyncPointName, bool bInWarpTranslation, bool bInIgnoreZAxis, bool bInWarpRotation, EMotionWarpRotationType InRotationType, float InWarpRotationTimeMultiplier)
 {
 	if (ensureAlways(InMotionWarpingComp))
 	{
@@ -175,6 +175,8 @@ void URootMotionModifierConfig_SkewWarp::AddRootMotionModifierSkewWarp(UMotionWa
 		NewModifier->bWarpTranslation = bInWarpTranslation;
 		NewModifier->bIgnoreZAxis = bInIgnoreZAxis;
 		NewModifier->bWarpRotation = bInWarpRotation;
+		NewModifier->RotationType = InRotationType;
+		NewModifier->WarpRotationTimeMultiplier = InWarpRotationTimeMultiplier;
 		InMotionWarpingComp->AddRootMotionModifier(NewModifier);
 	}
 }
\ No newline at end of file
diff --git a/Engine/Plugins/Experimental/Animation/MotionWarping/Source/MotionWarping/Public/RootMotionModifier.h b/Engine/Plugins/Experimental/Animation/MotionWarping/Source/MotionWarping/Public/RootMotionModifier.h
--- a/Engine/Plugins/Experimental/Animation/MotionWarping/Source/MotionWarping/Public/RootMotionModifier.h
+++ b/Engine/Plugins/Experimental/Animation/MotionWarping/Source/MotionWarping/Public/RootMotionModifier.h
@@ -133,6 +133,16 @@ protected:
 // FRootMotionModifier_Warp
 ///////////////////////////////////////////////////////////////
 
+UENUM(BlueprintType)
+enum class EMotionWarpRotationType : uint8
+{
+	/** Character rotates to match the rotation of the sync point */
+	Default,
+
+	/** Character rotates to face the sync point */
+	Facing,
+};
+
 USTRUCT()
 struct MOTIONWARPING_API FRootMotionModifier_Warp : public FRootMotionModifier
 {
@@ -156,6 +166,17 @@ public:
 	UPROPERTY()
 	bool bWarpRotation = true;
 
+	/** Whether rotation should be warp to match the rotation of the sync point or to face the sync point */
+	UPROPERTY()
+	EMotionWarpRotationType RotationType;
+
+	/** 
+	 * Allow to modify how fast the rotation is warped. 
+	 * e.g if the window duration is 2sec and this is 0.5, the target rotation will be reached in 1sec instead of 2sec
+	 */
+	UPROPERTY()
+	float WarpRotationTimeMultiplier = 1.f;
+
 	/** Sync Point used by this modifier as target for the warp. Cached during the Update */
 	UPROPERTY()
 	FMotionWarpingSyncPoint CachedSyncPoint;
@@ -178,6 +199,7 @@ public:
 
 protected:
 
+	FQuat GetTargetRotation(UMotionWarpingComponent& OwnerComp) const;
 	FQuat WarpRotation(UMotionWarpingComponent& OwnerComp, const FTransform& RootMotionDelta, const FTransform& RootMotionTotal, float DeltaSeconds);
 };
 
@@ -204,15 +226,26 @@ public:
 	UPROPERTY(EditAnywhere, BlueprintReadWrite, Category = "Config")
 	bool bWarpRotation = true;
 
+	/** Whether rotation should be warp to match the rotation of the sync point or to face the sync point */
+	UPROPERTY(EditAnywhere, BlueprintReadWrite, Category = "Config", meta = (EditCondition = "bWarpRotation"))
+	EMotionWarpRotationType RotationType;
+
+	/**
+	 * Allow to modify how fast the rotation is warped.
+	 * e.g if the window duration is 2sec and this is 0.5, the target rotation will be reached in 1sec instead of 2sec
+	 */
+	UPROPERTY(EditAnywhere, BlueprintReadWrite, Category = "Config", meta = (EditCondition = "bWarpRotation"))
+	float WarpRotationTimeMultiplier = 1.f;
+
 	URootMotionModifierConfig_Warp(const FObjectInitializer& ObjectInitializer) : Super(ObjectInitializer) {}
 
 	virtual void AddRootMotionModifier(UMotionWarpingComponent* MotionWarpingComp, const UAnimSequenceBase* Animation, float StartTime, float EndTime) const override
 	{
-		URootMotionModifierConfig_Warp::AddRootMotionModifierSimpleWarp(MotionWarpingComp, Animation, StartTime, EndTime, SyncPointName, bWarpTranslation, bIgnoreZAxis, bWarpRotation);
+		URootMotionModifierConfig_Warp::AddRootMotionModifierSimpleWarp(MotionWarpingComp, Animation, StartTime, EndTime, SyncPointName, bWarpTranslation, bIgnoreZAxis, bWarpRotation, RotationType, WarpRotationTimeMultiplier);
 	}
 
 	UFUNCTION(BlueprintCallable, Category = "Motion Warping")
-	static void AddRootMotionModifierSimpleWarp(UMotionWarpingComponent* InMotionWarpingComp, const UAnimSequenceBase* InAnimation, float InStartTime, float InEndTime, FName InSyncPointName, bool bInWarpTranslation, bool bInIgnoreZAxis, bool bInWarpRotation);
+	static void AddRootMotionModifierSimpleWarp(UMotionWarpingComponent* InMotionWarpingComp, const UAnimSequenceBase* InAnimation, float InStartTime, float InEndTime, FName InSyncPointName, bool bInWarpTranslation, bool bInIgnoreZAxis, bool bInWarpRotation, EMotionWarpRotationType InRotationType, float InWarpRotationTimeMultiplier = 1.f);
 };
 
 // FRootMotionModifier_Scale
diff --git a/Engine/Plugins/Experimental/Animation/MotionWarping/Source/MotionWarping/Public/RootMotionModifier_SkewWarp.h b/Engine/Plugins/Experimental/Animation/MotionWarping/Source/MotionWarping/Public/RootMotionModifier_SkewWarp.h
--- a/Engine/Plugins/Experimental/Animation/MotionWarping/Source/MotionWarping/Public/RootMotionModifier_SkewWarp.h
+++ b/Engine/Plugins/Experimental/Animation/MotionWarping/Source/MotionWarping/Public/RootMotionModifier_SkewWarp.h
@@ -32,9 +32,9 @@ public:
 
 	virtual void AddRootMotionModifier(UMotionWarpingComponent* MotionWarpingComp, const UAnimSequenceBase* Animation, float StartTime, float EndTime) const override
 	{
-		URootMotionModifierConfig_SkewWarp::AddRootMotionModifierSkewWarp(MotionWarpingComp, Animation, StartTime, EndTime, SyncPointName, bWarpTranslation, bIgnoreZAxis, bWarpRotation);
+		URootMotionModifierConfig_SkewWarp::AddRootMotionModifierSkewWarp(MotionWarpingComp, Animation, StartTime, EndTime, SyncPointName, bWarpTranslation, bIgnoreZAxis, bWarpRotation, RotationType, WarpRotationTimeMultiplier);
 	}
 
 	UFUNCTION(BlueprintCallable, Category = "Motion Warping")
-	static void AddRootMotionModifierSkewWarp(UMotionWarpingComponent* InMotionWarpingComp, const UAnimSequenceBase* InAnimation, float InStartTime, float InEndTime, FName InSyncPointName, bool bInWarpTranslation, bool bInIgnoreZAxis, bool bInWarpRotation);
+	static void AddRootMotionModifierSkewWarp(UMotionWarpingComponent* InMotionWarpingComp, const UAnimSequenceBase* InAnimation, float InStartTime, float InEndTime, FName InSyncPointName, bool bInWarpTranslation, bool bInIgnoreZAxis, bool bInWarpRotation, EMotionWarpRotationType InRotationType, float InWarpRotationTimeMultiplier = 1.f);
 };
\ No newline at end of file
