Motion Warping - Fixed initialization warnings #review-15319695 [REDACTED_HANDLE] #fyi [REDACTED_HANDLE]

diff --git a/Engine/Plugins/Experimental/Animation/MotionWarping/Source/MotionWarping/Public/RootMotionModifier.h b/Engine/Plugins/Experimental/Animation/MotionWarping/Source/MotionWarping/Public/RootMotionModifier.h
--- a/Engine/Plugins/Experimental/Animation/MotionWarping/Source/MotionWarping/Public/RootMotionModifier.h
+++ b/Engine/Plugins/Experimental/Animation/MotionWarping/Source/MotionWarping/Public/RootMotionModifier.h
@@ -99,7 +99,8 @@ struct MOTIONWARPING_API FMotionWarpingSyncPoint
 {
 	GENERATED_BODY()
 
-	FMotionWarpingSyncPoint() {}
+	FMotionWarpingSyncPoint()
+		: Location(FVector::ZeroVector), Rotation(FQuat::Identity) {}
 	FMotionWarpingSyncPoint(const FVector& InLocation, const FQuat& InRotation)
 		: Location(InLocation), Rotation(InRotation) {}
 	FMotionWarpingSyncPoint(const FVector& InLocation, const FRotator& InRotation)
@@ -168,7 +169,7 @@ public:
 
 	/** Whether rotation should be warp to match the rotation of the sync point or to face the sync point */
 	UPROPERTY()
-	EMotionWarpRotationType RotationType;
+	EMotionWarpRotationType RotationType = EMotionWarpRotationType::Default;
 
 	/** 
 	 * Allow to modify how fast the rotation is warped. 
@@ -260,7 +261,7 @@ public:
 
 	/** Vector used to scale each component of the translation */
 	UPROPERTY()
-	FVector Scale;
+	FVector Scale = FVector(1.f);
 
 	FRootMotionModifier_Scale() { bInLocalSpace = true; }
 	virtual ~FRootMotionModifier_Scale() {}
diff --git a/Engine/Plugins/Experimental/Animation/MotionWarping/Source/MotionWarping/Public/RootMotionModifier_AdjustmentBlendWarp.h b/Engine/Plugins/Experimental/Animation/MotionWarping/Source/MotionWarping/Public/RootMotionModifier_AdjustmentBlendWarp.h
--- a/Engine/Plugins/Experimental/Animation/MotionWarping/Source/MotionWarping/Public/RootMotionModifier_AdjustmentBlendWarp.h
+++ b/Engine/Plugins/Experimental/Animation/MotionWarping/Source/MotionWarping/Public/RootMotionModifier_AdjustmentBlendWarp.h
@@ -21,10 +21,10 @@ struct FMotionDeltaTrack
 	TArray<FRotator> DeltaRotationTrack;
 
 	UPROPERTY()
-	FVector TotalTranslation;
+	FVector TotalTranslation = FVector::ZeroVector;
 
 	UPROPERTY()
-	FRotator TotalRotation;
+	FRotator TotalRotation = FRotator::ZeroRotator;
 };
 
 USTRUCT()
@@ -49,7 +49,7 @@ struct MOTIONWARPING_API FRootMotionModifier_AdjustmentBlendWarp : public FRootM
 public:
 
 	UPROPERTY()
-	bool bWarpIKBones;
+	bool bWarpIKBones = false;
 
 	UPROPERTY()
 	TArray<FName> IKBones;
@@ -78,7 +78,7 @@ protected:
 	FAnimSequenceTrackContainer Result;
 
 	UPROPERTY()
-	float ActualStartTime;
+	float ActualStartTime = 0.f;
 
 	void PrecomputeWarpedTracks(UMotionWarpingComponent& OwnerComp);
 
