Motion Warping - Added draw debug features to visualize original motion, warped motion and actual motion. Added helper functions to get all Motion Warping windows from a given animation. #jira none #review-15509437 [REDACTED_HANDLE] #fyi [REDACTED_HANDLE]

diff --git a/Engine/Plugins/Experimental/Animation/MotionWarping/Source/MotionWarping/Private/MotionWarpingComponent.cpp b/Engine/Plugins/Experimental/Animation/MotionWarping/Source/MotionWarping/Private/MotionWarpingComponent.cpp
--- a/Engine/Plugins/Experimental/Animation/MotionWarping/Source/MotionWarping/Private/MotionWarpingComponent.cpp
+++ b/Engine/Plugins/Experimental/Animation/MotionWarping/Source/MotionWarping/Private/MotionWarpingComponent.cpp
@@ -78,6 +78,54 @@ FTransform UMotionWarpingUtilities::ExtractRootMotionFromAnimation(const UAnimSe
 	return FTransform::Identity;
 }
 
+void UMotionWarpingUtilities::GetMotionWarpingWindowsFromAnimation(const UAnimSequenceBase* Animation, TArray<FMotionWarpingWindowData>& OutWindows)
+{
+	if(Animation)
+	{
+		OutWindows.Reset();
+
+		for (int32 Idx = 0; Idx < Animation->Notifies.Num(); Idx++)
+		{
+			const FAnimNotifyEvent& NotifyEvent = Animation->Notifies[Idx];
+			if (UAnimNotifyState_MotionWarping* Notify = Cast<UAnimNotifyState_MotionWarping>(NotifyEvent.NotifyStateClass))
+			{
+				FMotionWarpingWindowData Data;
+				Data.AnimNotify = Notify;
+				Data.StartTime = NotifyEvent.GetTriggerTime();
+				Data.EndTime = NotifyEvent.GetEndTriggerTime();
+				OutWindows.Add(Data);
+			}
+		}
+	}
+}
+
+void UMotionWarpingUtilities::GetMotionWarpingWindowsForSyncPointFromAnimation(const UAnimSequenceBase* Animation, FName SyncPointName, TArray<FMotionWarpingWindowData>& OutWindows)
+{
+	if (Animation && SyncPointName != NAME_None)
+	{
+		OutWindows.Reset();
+
+		for (int32 Idx = 0; Idx < Animation->Notifies.Num(); Idx++)
+		{
+			const FAnimNotifyEvent& NotifyEvent = Animation->Notifies[Idx];
+			if (UAnimNotifyState_MotionWarping* Notify = Cast<UAnimNotifyState_MotionWarping>(NotifyEvent.NotifyStateClass))
+			{
+				if (const URootMotionModifierConfig_Warp* Config = Cast<const URootMotionModifierConfig_Warp>(Notify->RootMotionModifierConfig))
+				{
+					if(Config->SyncPointName == SyncPointName)
+					{
+						FMotionWarpingWindowData Data;
+						Data.AnimNotify = Notify;
+						Data.StartTime = NotifyEvent.GetTriggerTime();
+						Data.EndTime = NotifyEvent.GetEndTriggerTime();
+						OutWindows.Add(Data);
+					}
+				}
+			}
+		}
+	}
+}
+
 // UMotionWarpingComponent
 ///////////////////////////////////////////////////////////////////////
 
@@ -270,6 +318,37 @@ FTransform UMotionWarpingComponent::ProcessRootMotionPostConvertToWorld(const FT
 		}
 	}
 
+#if !(UE_BUILD_SHIPPING || UE_BUILD_TEST)
+	const int32 DebugLevel = FMotionWarpingCVars::CVarMotionWarpingDebug.GetValueOnGameThread();
+	if (DebugLevel >= 2)
+	{
+		const float DrawDebugDuration = FMotionWarpingCVars::CVarMotionWarpingDrawDebugDuration.GetValueOnGameThread();
+		const float PointSize = 7.f;
+		const FVector ActorFeetLocation = CharacterMovementComponent->GetActorFeetLocation();
+		if (RootMotionModifiers.Num() > 0)
+		{
+			if (!OriginalRootMotionAccum.IsSet())
+			{
+				OriginalRootMotionAccum = ActorFeetLocation;
+				WarpedRootMotionAccum = ActorFeetLocation;
+			}
+			
+			OriginalRootMotionAccum = OriginalRootMotionAccum.GetValue() + InRootMotion.GetLocation();
+			WarpedRootMotionAccum = WarpedRootMotionAccum.GetValue() + FinalRootMotion.GetLocation();
+
+			DrawDebugPoint(GetWorld(), OriginalRootMotionAccum.GetValue(), PointSize, FColor::Red, false, DrawDebugDuration, 0);
+			DrawDebugPoint(GetWorld(), WarpedRootMotionAccum.GetValue(), PointSize, FColor::Green, false, DrawDebugDuration, 0);
+		}
+		else
+		{
+			OriginalRootMotionAccum.Reset();
+			WarpedRootMotionAccum.Reset();
+		}
+
+		DrawDebugPoint(GetWorld(), ActorFeetLocation, PointSize, FColor::Blue, false, DrawDebugDuration, 0);
+	}
+#endif
+
 	return FinalRootMotion;
 }
 
diff --git a/Engine/Plugins/Experimental/Animation/MotionWarping/Source/MotionWarping/Public/AnimNotifyState_MotionWarping.h b/Engine/Plugins/Experimental/Animation/MotionWarping/Source/MotionWarping/Public/AnimNotifyState_MotionWarping.h
--- a/Engine/Plugins/Experimental/Animation/MotionWarping/Source/MotionWarping/Public/AnimNotifyState_MotionWarping.h
+++ b/Engine/Plugins/Experimental/Animation/MotionWarping/Source/MotionWarping/Public/AnimNotifyState_MotionWarping.h
@@ -14,7 +14,7 @@ class MOTIONWARPING_API UAnimNotifyState_MotionWarping : public UAnimNotifyState
 
 public:
 
-	UPROPERTY(EditAnywhere, Instanced, Category = "Config")
+	UPROPERTY(EditAnywhere, Instanced, BlueprintReadWrite, Category = "Config")
 	class URootMotionModifierConfig* RootMotionModifierConfig;
 
 	UAnimNotifyState_MotionWarping(const FObjectInitializer& ObjectInitializer);
diff --git a/Engine/Plugins/Experimental/Animation/MotionWarping/Source/MotionWarping/Public/MotionWarpingComponent.h b/Engine/Plugins/Experimental/Animation/MotionWarping/Source/MotionWarping/Public/MotionWarpingComponent.h
--- a/Engine/Plugins/Experimental/Animation/MotionWarping/Source/MotionWarping/Public/MotionWarpingComponent.h
+++ b/Engine/Plugins/Experimental/Animation/MotionWarping/Source/MotionWarping/Public/MotionWarpingComponent.h
@@ -12,6 +12,7 @@ class ACharacter;
 class UAnimSequenceBase;
 class UCharacterMovementComponent;
 class UMotionWarpingComponent;
+class UAnimNotifyState_MotionWarping;
 
 DECLARE_LOG_CATEGORY_EXTERN(LogMotionWarping, Log, All);
 
@@ -24,6 +25,21 @@ struct FMotionWarpingCVars
 };
 #endif
 
+USTRUCT(BlueprintType)
+struct FMotionWarpingWindowData
+{
+	GENERATED_BODY()
+
+	UPROPERTY(EditAnywhere, BlueprintReadWrite, Category = "Defaults")
+	TObjectPtr<UAnimNotifyState_MotionWarping> AnimNotify;
+
+	UPROPERTY(EditAnywhere, BlueprintReadWrite, Category = "Defaults")
+	float StartTime = 0.f;
+
+	UPROPERTY(EditAnywhere, BlueprintReadWrite, Category = "Defaults")
+	float EndTime = 0.f;
+};
+
 UCLASS()
 class MOTIONWARPING_API UMotionWarpingUtilities : public UBlueprintFunctionLibrary
 {
@@ -48,6 +64,14 @@ public:
 	/** Extract Root Motion transform from a contiguous position range */
 	UFUNCTION(BlueprintCallable, Category = "Motion Warping")
 	static FTransform ExtractRootMotionFromAnimation(const UAnimSequenceBase* Animation, float StartTime, float EndTime);
+
+	/** @return All the MotionWarping windows within the supplied animation */
+	UFUNCTION(BlueprintCallable, Category = "Motion Warping")
+	static void GetMotionWarpingWindowsFromAnimation(const UAnimSequenceBase* Animation, TArray<FMotionWarpingWindowData>& OutWindows);
+
+	/** @return All the MotionWarping windows within the supplied animation for a given Sync Point */
+	UFUNCTION(BlueprintCallable, Category = "Motion Warping")
+	static void GetMotionWarpingWindowsForSyncPointFromAnimation(const UAnimSequenceBase* Animation, FName SyncPointName, TArray<FMotionWarpingWindowData>& OutWindows);
 };
 
 DECLARE_DYNAMIC_MULTICAST_DELEGATE_OneParam(FMotionWarpingPreUpdate, class UMotionWarpingComponent*, MotionWarpingComp);
@@ -109,6 +133,11 @@ protected:
 	UPROPERTY(Transient)
 	TMap<FName, FMotionWarpingSyncPoint> SyncPoints;
 
+#if !(UE_BUILD_SHIPPING || UE_BUILD_TEST)
+	TOptional<FVector> OriginalRootMotionAccum;
+	TOptional<FVector> WarpedRootMotionAccum;
+#endif
+
 	void Update();
 
 	FTransform ProcessRootMotionPreConvertToWorld(const FTransform& InRootMotion, class UCharacterMovementComponent* CharacterMovementComponent, float DeltaSeconds);
