Motion Warping
- Added accessors for modifier's state and overridable event called when the state changes (state member will be private shortly)
- Temp solution to notify AnimNotifyState when warp starts, end and updates

diff --git a/Engine/Plugins/Experimental/Animation/MotionWarping/Source/MotionWarping/Private/AnimNotifyState_MotionWarping.cpp b/Engine/Plugins/Experimental/Animation/MotionWarping/Source/MotionWarping/Private/AnimNotifyState_MotionWarping.cpp
--- a/Engine/Plugins/Experimental/Animation/MotionWarping/Source/MotionWarping/Private/AnimNotifyState_MotionWarping.cpp
+++ b/Engine/Plugins/Experimental/Animation/MotionWarping/Source/MotionWarping/Private/AnimNotifyState_MotionWarping.cpp
@@ -10,7 +10,7 @@ UAnimNotifyState_MotionWarping::UAnimNotifyState_MotionWarping(const FObjectInit
 {
 }
 
-void UAnimNotifyState_MotionWarping::AddRootMotionModifier_Implementation(UMotionWarpingComponent* MotionWarpingComp, UAnimSequenceBase* Animation, float StartTime, float EndTime) const
+void UAnimNotifyState_MotionWarping::AddRootMotionModifier_Implementation(UMotionWarpingComponent* MotionWarpingComp, const UAnimSequenceBase* Animation, float StartTime, float EndTime) const
 {
 	if (MotionWarpingComp && RootMotionModifierConfig)
 	{
diff --git a/Engine/Plugins/Experimental/Animation/MotionWarping/Source/MotionWarping/Private/MotionWarpingComponent.cpp b/Engine/Plugins/Experimental/Animation/MotionWarping/Source/MotionWarping/Private/MotionWarpingComponent.cpp
--- a/Engine/Plugins/Experimental/Animation/MotionWarping/Source/MotionWarping/Private/MotionWarpingComponent.cpp
+++ b/Engine/Plugins/Experimental/Animation/MotionWarping/Source/MotionWarping/Private/MotionWarpingComponent.cpp
@@ -175,7 +175,7 @@ void UMotionWarpingComponent::DisableAllRootMotionModifiers()
 	{
 		for (TSharedPtr<FRootMotionModifier>& Modifier : RootMotionModifiers)
 		{
-			Modifier->State = ERootMotionModifierState::Disabled;
+			Modifier->SetState(ERootMotionModifierState::Disabled);
 		}
 	}
 }
@@ -205,6 +205,17 @@ void UMotionWarpingComponent::Update()
 					if (!ContainsModifier(Montage, StartTime, EndTime))
 					{
 						MotionWarpingNotify->AddRootMotionModifier(this, Montage, StartTime, EndTime);
+
+						//@TODO: Temp hack to keep track of the AnimNotifyState each modifier is created from.
+						if (RootMotionModifiers.Num())
+						{
+							TSharedPtr<FRootMotionModifier> Last = RootMotionModifiers.Last();
+							if (Last.IsValid() && Last->Animation.Get() == Montage && Last->StartTime == StartTime && Last->EndTime == EndTime)
+							{
+								Last->AnimNotifyState = MotionWarpingNotify;
+								Last->AnimNotifyState->OnWarpBegin(this, Last->Animation.Get(), Last->StartTime, Last->EndTime);
+							}
+						}
 					}
 				}
 			}
@@ -236,6 +247,17 @@ void UMotionWarpingComponent::Update()
 								if (!ContainsModifier(Montage, StartTime, EndTime))
 								{
 									MotionWarpingNotify->AddRootMotionModifier(this, Montage, StartTime, EndTime);
+
+									//@TODO: Temp hack to keep track of the AnimNotifyState each modifier is created from.
+									if (RootMotionModifiers.Num())
+									{
+										TSharedPtr<FRootMotionModifier> Last = RootMotionModifiers.Last();
+										if(Last.IsValid() && Last->Animation.Get() == Montage && Last->StartTime == StartTime && Last->EndTime == EndTime)
+										{
+											Last->AnimNotifyState = MotionWarpingNotify;
+											Last->AnimNotifyState->OnWarpBegin(this, Last->Animation.Get(), Last->StartTime, Last->EndTime);
+										}
+									}
 								}
 							}
 						}
@@ -252,13 +274,27 @@ void UMotionWarpingComponent::Update()
 	{
 		for (TSharedPtr<FRootMotionModifier>& Modifier : RootMotionModifiers)
 		{
+			const UAnimNotifyState_MotionWarping* AnimNotify = Modifier->AnimNotifyState.Get();
+			if (AnimNotify)
+			{
+				AnimNotify->OnWarpPreUpdate(this, Modifier->Animation.Get(), Modifier->StartTime, Modifier->EndTime);
+			}
+
 			Modifier->Update(*this);
+
+			if(AnimNotify)
+			{
+				if(Modifier->GetState() == ERootMotionModifierState::Disabled || Modifier->GetState() == ERootMotionModifierState::MarkedForRemoval)
+				{
+					AnimNotify->OnWarpEnd(this, Modifier->Animation.Get(), Modifier->StartTime, Modifier->EndTime);
+				}
+			}
 		}
 
 		// Remove the modifiers that has been marked for removal
 		RootMotionModifiers.RemoveAll([this](const TSharedPtr<FRootMotionModifier>& Modifier) 
 		{ 
-			if(Modifier->State == ERootMotionModifierState::MarkedForRemoval)
+			if(Modifier->GetState() == ERootMotionModifierState::MarkedForRemoval)
 			{
 				UE_LOG(LogMotionWarping, Verbose, TEXT("MotionWarping: RootMotionModifier removed. NetMode: %d WorldTime: %f Char: %s Animation: %s [%f %f] [%f %f] Loc: %s Rot: %s"),
 					GetWorld()->GetNetMode(), GetWorld()->GetTimeSeconds(), *GetNameSafe(GetCharacterOwner()), *GetNameSafe(Modifier->Animation.Get()), Modifier->StartTime, Modifier->EndTime, Modifier->PreviousPosition, Modifier->CurrentPosition,
@@ -289,7 +325,7 @@ FTransform UMotionWarpingComponent::ProcessRootMotionPreConvertToWorld(const FTr
 	// Apply Local Space Modifiers
 	for (TSharedPtr<FRootMotionModifier>& RootMotionModifier : RootMotionModifiers)
 	{
-		if (RootMotionModifier->State == ERootMotionModifierState::Active && RootMotionModifier->bInLocalSpace)
+		if (RootMotionModifier->GetState() == ERootMotionModifierState::Active && RootMotionModifier->bInLocalSpace)
 		{
 			FinalRootMotion = RootMotionModifier->ProcessRootMotion(*this, FinalRootMotion, DeltaSeconds);
 		}
@@ -312,7 +348,7 @@ FTransform UMotionWarpingComponent::ProcessRootMotionPostConvertToWorld(const FT
 	// Apply World Space Modifiers
 	for (TSharedPtr<FRootMotionModifier>& RootMotionModifier : RootMotionModifiers)
 	{
-		if (RootMotionModifier->State == ERootMotionModifierState::Active && !RootMotionModifier->bInLocalSpace)
+		if (RootMotionModifier->GetState() == ERootMotionModifierState::Active && !RootMotionModifier->bInLocalSpace)
 		{
 			FinalRootMotion = RootMotionModifier->ProcessRootMotion(*this, FinalRootMotion, DeltaSeconds);
 		}
diff --git a/Engine/Plugins/Experimental/Animation/MotionWarping/Source/MotionWarping/Private/RootMotionModifier.cpp b/Engine/Plugins/Experimental/Animation/MotionWarping/Source/MotionWarping/Private/RootMotionModifier.cpp
--- a/Engine/Plugins/Experimental/Animation/MotionWarping/Source/MotionWarping/Private/RootMotionModifier.cpp
+++ b/Engine/Plugins/Experimental/Animation/MotionWarping/Source/MotionWarping/Private/RootMotionModifier.cpp
@@ -7,6 +7,8 @@
 #include "Animation/AnimMontage.h"
 #include "MotionWarpingComponent.h"
 #include "DrawDebugHelpers.h"
+#include "Animation/AnimInstance.h"
+#include "AnimNotifyState_MotionWarping.h"
 
 // FRootMotionModifier
 ///////////////////////////////////////////////////////////////
@@ -22,7 +24,7 @@ void FRootMotionModifier::Update(UMotionWarpingComponent& OwnerComp)
 		UE_LOG(LogMotionWarping, Verbose, TEXT("MotionWarping: Marking RootMotionModifier for removal. Reason: Animation is not valid. Char: %s Current Montage: %s. Window: Animation: %s [%f %f] [%f %f]"),
 			*GetNameSafe(OwnerComp.GetCharacterOwner()), *GetNameSafe(Montage), *GetNameSafe(Animation.Get()), StartTime, EndTime, PreviousPosition, CurrentPosition);
 
-		State = ERootMotionModifierState::MarkedForRemoval;
+		SetState(ERootMotionModifierState::MarkedForRemoval);
 		return;
 	}
 
@@ -37,7 +39,7 @@ void FRootMotionModifier::Update(UMotionWarpingComponent& OwnerComp)
 		UE_LOG(LogMotionWarping, Verbose, TEXT("MotionWarping: Marking RootMotionModifier for removal. Reason: Window has ended. Char: %s Animation: %s [%f %f] [%f %f]"),
 			*GetNameSafe(OwnerComp.GetCharacterOwner()), *GetNameSafe(Animation.Get()), StartTime, EndTime, PreviousPosition, CurrentPosition);
 
-		State = ERootMotionModifierState::MarkedForRemoval;
+		SetState(ERootMotionModifierState::MarkedForRemoval);
 		return;
 	}
 
@@ -45,13 +47,25 @@ void FRootMotionModifier::Update(UMotionWarpingComponent& OwnerComp)
 	if (PreviousPosition >= StartTime && PreviousPosition < EndTime)
 	{
 		// If we were waiting, switch to active
-		if (State == ERootMotionModifierState::Waiting)
+		if (GetState() == ERootMotionModifierState::Waiting)
 		{
-			State = ERootMotionModifierState::Active;
+			SetState(ERootMotionModifierState::Active);
 		}
 	}
 }
 
+void FRootMotionModifier::SetState(ERootMotionModifierState NewState)
+{
+	if(State != NewState)
+	{
+		ERootMotionModifierState LastState = State;
+
+		State = NewState;
+
+		OnStateChanged(LastState);
+	}
+}
+
 // FRootMotionModifier_Warp
 ///////////////////////////////////////////////////////////////
 
@@ -61,7 +75,7 @@ void FRootMotionModifier_Warp::Update(UMotionWarpingComponent& OwnerComp)
 	FRootMotionModifier::Update(OwnerComp);
 
 	// Cache sync point transform and trigger OnSyncPointChanged if needed
-	if (State == ERootMotionModifierState::Active)
+	if (GetState() == ERootMotionModifierState::Active)
 	{
 		const FMotionWarpingSyncPoint* SyncPointPtr = OwnerComp.FindSyncPoint(SyncPointName);
 
@@ -71,7 +85,7 @@ void FRootMotionModifier_Warp::Update(UMotionWarpingComponent& OwnerComp)
 			UE_LOG(LogMotionWarping, Verbose, TEXT("MotionWarping: Marking RootMotionModifier as Disabled. Reason: Invalid Sync Point (%s). Char: %s Animation: %s [%f %f] [%f %f]"),
 				*SyncPointName.ToString(), *GetNameSafe(OwnerComp.GetCharacterOwner()), *GetNameSafe(Animation.Get()), StartTime, EndTime, PreviousPosition, CurrentPosition);
 
-			State = ERootMotionModifierState::Disabled;
+			SetState(ERootMotionModifierState::Disabled);
 			return;
 		}
 
diff --git a/Engine/Plugins/Experimental/Animation/MotionWarping/Source/MotionWarping/Private/RootMotionModifier_AdjustmentBlendWarp.cpp b/Engine/Plugins/Experimental/Animation/MotionWarping/Source/MotionWarping/Private/RootMotionModifier_AdjustmentBlendWarp.cpp
--- a/Engine/Plugins/Experimental/Animation/MotionWarping/Source/MotionWarping/Private/RootMotionModifier_AdjustmentBlendWarp.cpp
+++ b/Engine/Plugins/Experimental/Animation/MotionWarping/Source/MotionWarping/Private/RootMotionModifier_AdjustmentBlendWarp.cpp
@@ -407,7 +407,7 @@ void URootMotionModifierConfig_AdjustmentBlendWarp::GetIKBoneTransformAndAlpha(A
 	{
 		for (const auto& Modifier : MotionWarpingComp->GetRootMotionModifiers())
 		{
-			if (Modifier.IsValid() && Modifier->State == ERootMotionModifierState::Active && Modifier->GetScriptStruct()->IsChildOf(FRootMotionModifier_AdjustmentBlendWarp::StaticStruct()))
+			if (Modifier.IsValid() && Modifier->GetState() == ERootMotionModifierState::Active && Modifier->GetScriptStruct()->IsChildOf(FRootMotionModifier_AdjustmentBlendWarp::StaticStruct()))
 			{
 				// We must check if the Animation for the modifier is still relevant because in RootMotionFromMontageOnlyMode the montage could be aborted 
 				// but the modifier will remain in the list until the next update
diff --git a/Engine/Plugins/Experimental/Animation/MotionWarping/Source/MotionWarping/Public/AnimNotifyState_MotionWarping.h b/Engine/Plugins/Experimental/Animation/MotionWarping/Source/MotionWarping/Public/AnimNotifyState_MotionWarping.h
--- a/Engine/Plugins/Experimental/Animation/MotionWarping/Source/MotionWarping/Public/AnimNotifyState_MotionWarping.h
+++ b/Engine/Plugins/Experimental/Animation/MotionWarping/Source/MotionWarping/Public/AnimNotifyState_MotionWarping.h
@@ -6,6 +6,10 @@
 #include "Animation/AnimNotifies/AnimNotifyState.h"
 #include "AnimNotifyState_MotionWarping.generated.h"
 
+class UMotionWarpingComponent;
+class UAnimSequenceBase;
+class URootMotionModifierConfig;
+
 /** AnimNotifyState used to define a motion warping window in the animation */
 UCLASS(meta = (DisplayName = "Motion Warping"))
 class MOTIONWARPING_API UAnimNotifyState_MotionWarping : public UAnimNotifyState
@@ -15,10 +19,19 @@ class MOTIONWARPING_API UAnimNotifyState_MotionWarping : public UAnimNotifyState
 public:
 
 	UPROPERTY(EditAnywhere, Instanced, BlueprintReadWrite, Category = "Config")
-	class URootMotionModifierConfig* RootMotionModifierConfig;
+	URootMotionModifierConfig* RootMotionModifierConfig;
 
 	UAnimNotifyState_MotionWarping(const FObjectInitializer& ObjectInitializer);
 
 	UFUNCTION(BlueprintNativeEvent, Category = "Motion Warping")
-	void AddRootMotionModifier(class UMotionWarpingComponent* MotionWarpingComp, class UAnimSequenceBase* Animation, float StartTime, float EndTime) const;
+	void AddRootMotionModifier(UMotionWarpingComponent* MotionWarpingComp, const UAnimSequenceBase* Animation, float StartTime, float EndTime) const;
+
+	UFUNCTION(BlueprintImplementableEvent, Category = "Motion Warping")
+	void OnWarpBegin(UMotionWarpingComponent* MotionWarpingComp, const UAnimSequenceBase* Animation, float StartTime, float EndTime) const;
+
+	UFUNCTION(BlueprintImplementableEvent, Category = "Motion Warping")
+	void OnWarpPreUpdate(UMotionWarpingComponent* MotionWarpingComp, const UAnimSequenceBase* Animation, float StartTime, float EndTime) const;
+
+	UFUNCTION(BlueprintImplementableEvent, Category = "Motion Warping")
+	void OnWarpEnd(UMotionWarpingComponent* MotionWarpingComp, const UAnimSequenceBase* Animation, float StartTime, float EndTime) const;
 };
\ No newline at end of file
diff --git a/Engine/Plugins/Experimental/Animation/MotionWarping/Source/MotionWarping/Public/RootMotionModifier.h b/Engine/Plugins/Experimental/Animation/MotionWarping/Source/MotionWarping/Public/RootMotionModifier.h
--- a/Engine/Plugins/Experimental/Animation/MotionWarping/Source/MotionWarping/Public/RootMotionModifier.h
+++ b/Engine/Plugins/Experimental/Animation/MotionWarping/Source/MotionWarping/Public/RootMotionModifier.h
@@ -7,6 +7,7 @@
 #include "RootMotionModifier.generated.h"
 
 class UMotionWarpingComponent;
+class UAnimNotifyState_MotionWarping;
 
 /** The possible states of a Root Motion Modifier */
 UENUM(BlueprintType)
@@ -37,6 +38,10 @@ public:
 	UPROPERTY()
 	TWeakObjectPtr<const UAnimSequenceBase> Animation = nullptr;
 
+	/** AnimNotifyState this modifier was created from. May be NULL */
+	UPROPERTY()
+	TWeakObjectPtr<const UAnimNotifyState_MotionWarping> AnimNotifyState = nullptr;
+
 	/** Start time of the warping window */
 	UPROPERTY()
 	float StartTime = 0.f;
@@ -63,7 +68,7 @@ public:
 
 	/** Current state */
 	UPROPERTY()
-	ERootMotionModifierState State = ERootMotionModifierState::Waiting;
+	ERootMotionModifierState State = ERootMotionModifierState::Waiting; //@TODO: Move to private
 
 	FRootMotionModifier(){}
 	virtual ~FRootMotionModifier() {}
@@ -76,6 +81,15 @@ public:
 
 	/** Performs the actual modification to the motion */
 	virtual FTransform ProcessRootMotion(UMotionWarpingComponent& OwnerComp, const FTransform& InRootMotion, float DeltaSeconds) PURE_VIRTUAL(FRootMotionModifier::ProcessRootMotion, return FTransform::Identity;);
+
+	/** Called when the state of the modifier changes */
+	virtual void OnStateChanged(ERootMotionModifierState LastState){}
+
+	/** Sets the state of the modifier */
+	void SetState(ERootMotionModifierState NewState);
+
+	/** Returns the state of the modifier */
+	FORCEINLINE ERootMotionModifierState GetState() const { return State; }
 };
 
 /** Blueprint wrapper around the config properties of a root motion modifier */
