Motion Warping
 - Added option to allow users to define a warp point in the animation.
 - Added delegates to notify external objects when a modifier updates and changes state
 - AnimNotifyState now listen to various events from the modifier to trigger the corresponding blueprint events.
 - Added RootMotionModifierHandle to uniquely identify each modifier
 - Added function to get a motion modifier by a handle
 - Added blueprint function to bind blueprint events to motion modifier's events
 - Added ExtractRootTransform helper function to MotionWArpingUtilities
 - Deprecated functions receiving OwnerComp via param

diff --git a/Engine/Plugins/Experimental/Animation/MotionWarping/Source/MotionWarping/Private/AnimNotifyState_MotionWarping.cpp b/Engine/Plugins/Experimental/Animation/MotionWarping/Source/MotionWarping/Private/AnimNotifyState_MotionWarping.cpp
--- a/Engine/Plugins/Experimental/Animation/MotionWarping/Source/MotionWarping/Private/AnimNotifyState_MotionWarping.cpp
+++ b/Engine/Plugins/Experimental/Animation/MotionWarping/Source/MotionWarping/Private/AnimNotifyState_MotionWarping.cpp
@@ -3,17 +3,75 @@
 #include "AnimNotifyState_MotionWarping.h"
 #include "GameFramework/Actor.h"
 #include "MotionWarpingComponent.h"
-#include "RootMotionModifier.h"
 
 UAnimNotifyState_MotionWarping::UAnimNotifyState_MotionWarping(const FObjectInitializer& ObjectInitializer)
 	: Super(ObjectInitializer)
 {
 }
 
-void UAnimNotifyState_MotionWarping::AddRootMotionModifier_Implementation(UMotionWarpingComponent* MotionWarpingComp, const UAnimSequenceBase* Animation, float StartTime, float EndTime) const
+FRootMotionModifierHandle UAnimNotifyState_MotionWarping::OnBecomeRelevant(UMotionWarpingComponent* MotionWarpingComp, const UAnimSequenceBase* Animation, float StartTime, float EndTime) const
+{
+	FRootMotionModifierHandle Handle = AddRootMotionModifier(MotionWarpingComp, Animation, StartTime, EndTime);
+
+	//@TODO: Temp, to handle the case where the blueprint notify does not return the handle. Will be removed shortly
+	if(!Handle.IsValid())
+	{
+		if (MotionWarpingComp->GetRootMotionModifiers().Num())
+		{
+			TSharedPtr<FRootMotionModifier> Last = MotionWarpingComp->GetRootMotionModifiers().Last();
+			if (Last.IsValid() && Last->Animation.Get() == Animation && Last->StartTime == StartTime && Last->EndTime == EndTime)
+			{
+				Handle = Last->GetHandle();
+			}
+		}
+	}
+
+	TSharedPtr<FRootMotionModifier> Modifier = MotionWarpingComp->GetRootMotionModifierByHandle(Handle);
+	if (Modifier.IsValid())
+	{
+		if(!Modifier->OnActivateDelegate.IsBound())
+		{
+			Modifier->OnActivateDelegate.BindDynamic(this, &UAnimNotifyState_MotionWarping::OnRootMotionModifierActivate);
+		}
+
+		if (!Modifier->OnUpdateDelegate.IsBound())
+		{
+			Modifier->OnUpdateDelegate.BindDynamic(this, &UAnimNotifyState_MotionWarping::OnRootMotionModifierUpdate);
+		}
+
+		if (!Modifier->OnDeactivateDelegate.IsBound())
+		{
+			Modifier->OnDeactivateDelegate.BindDynamic(this, &UAnimNotifyState_MotionWarping::OnRootMotionModifierDeactivate);
+		}
+	}
+
+	return Handle;
+}
+
+FRootMotionModifierHandle UAnimNotifyState_MotionWarping::AddRootMotionModifier_Implementation(UMotionWarpingComponent* MotionWarpingComp, const UAnimSequenceBase* Animation, float StartTime, float EndTime) const
 {
 	if (MotionWarpingComp && RootMotionModifierConfig)
 	{
-		RootMotionModifierConfig->AddRootMotionModifier(MotionWarpingComp, Animation, StartTime, EndTime);
+		return RootMotionModifierConfig->AddRootMotionModifierNew(MotionWarpingComp, Animation, StartTime, EndTime);
 	}
+
+	return FRootMotionModifierHandle::InvalidHandle;
+}
+
+void UAnimNotifyState_MotionWarping::OnRootMotionModifierActivate(UMotionWarpingComponent* MotionWarpingComp, const FRootMotionModifierHandle& Handle)
+{
+	// Notify blueprint
+	OnWarpBegin(MotionWarpingComp, Handle);
+}
+
+void UAnimNotifyState_MotionWarping::OnRootMotionModifierUpdate(UMotionWarpingComponent* MotionWarpingComp, const FRootMotionModifierHandle& Handle)
+{
+	// Notify blueprint
+	OnWarpUpdate(MotionWarpingComp, Handle);
+}
+
+void UAnimNotifyState_MotionWarping::OnRootMotionModifierDeactivate(UMotionWarpingComponent* MotionWarpingComp, const FRootMotionModifierHandle& Handle)
+{
+	// Notify blueprint
+	OnWarpEnd(MotionWarpingComp, Handle);
 }
\ No newline at end of file
diff --git a/Engine/Plugins/Experimental/Animation/MotionWarping/Source/MotionWarping/Private/MotionWarpingComponent.cpp b/Engine/Plugins/Experimental/Animation/MotionWarping/Source/MotionWarping/Private/MotionWarpingComponent.cpp
--- a/Engine/Plugins/Experimental/Animation/MotionWarping/Source/MotionWarping/Private/MotionWarpingComponent.cpp
+++ b/Engine/Plugins/Experimental/Animation/MotionWarping/Source/MotionWarping/Private/MotionWarpingComponent.cpp
@@ -78,6 +78,27 @@ FTransform UMotionWarpingUtilities::ExtractRootMotionFromAnimation(const UAnimSe
 	return FTransform::Identity;
 }
 
+FTransform UMotionWarpingUtilities::ExtractRootTransformFromAnimation(const UAnimSequenceBase* Animation, float Time)
+{
+	if (const UAnimMontage* AnimMontage = Cast<UAnimMontage>(Animation))
+	{
+		if(const FAnimSegment* Segment = AnimMontage->SlotAnimTracks[0].AnimTrack.GetSegmentAtTime(Time))
+		{
+			if (const UAnimSequence* AnimSequence = Cast<UAnimSequence>(Segment->AnimReference))
+			{
+				const float AnimSequenceTime = Segment->ConvertTrackPosToAnimPos(Time);
+				return AnimSequence->ExtractRootTrackTransform(AnimSequenceTime, nullptr);
+			}	
+		}
+	}
+	else if (const UAnimSequence* AnimSequence = Cast<UAnimSequence>(Animation))
+	{
+		return AnimSequence->ExtractRootTrackTransform(Time, nullptr);
+	}
+
+	return FTransform::Identity;
+}
+
 void UMotionWarpingUtilities::GetMotionWarpingWindowsFromAnimation(const UAnimSequenceBase* Animation, TArray<FMotionWarpingWindowData>& OutWindows)
 {
 	if(Animation)
@@ -157,16 +178,39 @@ bool UMotionWarpingComponent::ContainsModifier(const UAnimSequenceBase* Animatio
 		});
 }
 
-void UMotionWarpingComponent::AddRootMotionModifier(TSharedPtr<FRootMotionModifier> Modifier)
+FRootMotionModifierHandle UMotionWarpingComponent::AddRootMotionModifier(TSharedPtr<FRootMotionModifier> Modifier)
 {
 	if (ensureAlways(Modifier.IsValid()))
 	{
+		//@TODO: Move this from here after update all modifiers
+		Modifier->Initialize(this);
+
 		RootMotionModifiers.Add(Modifier);
 
 		UE_LOG(LogMotionWarping, Verbose, TEXT("MotionWarping: RootMotionModifier added. NetMode: %d WorldTime: %f Char: %s Animation: %s [%f %f] [%f %f] Loc: %s Rot: %s"),
 			GetWorld()->GetNetMode(), GetWorld()->GetTimeSeconds(), *GetNameSafe(GetCharacterOwner()), *GetNameSafe(Modifier->Animation.Get()), Modifier->StartTime, Modifier->EndTime, Modifier->PreviousPosition, Modifier->CurrentPosition,
 			*GetCharacterOwner()->GetActorLocation().ToString(), *GetCharacterOwner()->GetActorRotation().ToCompactString());
+
+		return Modifier->GetHandle();
+	}
+
+	return FRootMotionModifierHandle::InvalidHandle;
+}
+
+TSharedPtr<FRootMotionModifier> UMotionWarpingComponent::GetRootMotionModifierByHandle(const FRootMotionModifierHandle& Handle) const
+{
+	if(Handle.IsValid())
+	{
+		for (const TSharedPtr<FRootMotionModifier>& RootMotionModifier : RootMotionModifiers)
+		{
+			if (RootMotionModifier.IsValid() && RootMotionModifier->GetHandle() == Handle)
+			{
+				return RootMotionModifier;
+			}
+		}
 	}
+
+	return TSharedPtr<FRootMotionModifier>(nullptr);
 }
 
 void UMotionWarpingComponent::DisableAllRootMotionModifiers()
@@ -204,18 +248,7 @@ void UMotionWarpingComponent::Update()
 				{
 					if (!ContainsModifier(Montage, StartTime, EndTime))
 					{
-						MotionWarpingNotify->AddRootMotionModifier(this, Montage, StartTime, EndTime);
-
-						//@TODO: Temp hack to keep track of the AnimNotifyState each modifier is created from.
-						if (RootMotionModifiers.Num())
-						{
-							TSharedPtr<FRootMotionModifier> Last = RootMotionModifiers.Last();
-							if (Last.IsValid() && Last->Animation.Get() == Montage && Last->StartTime == StartTime && Last->EndTime == EndTime)
-							{
-								Last->AnimNotifyState = MotionWarpingNotify;
-								Last->AnimNotifyState->OnWarpBegin(this, Last->Animation.Get(), Last->StartTime, Last->EndTime);
-							}
-						}
+						MotionWarpingNotify->OnBecomeRelevant(this, Montage, StartTime, EndTime);
 					}
 				}
 			}
@@ -246,18 +279,7 @@ void UMotionWarpingComponent::Update()
 							{
 								if (!ContainsModifier(Montage, StartTime, EndTime))
 								{
-									MotionWarpingNotify->AddRootMotionModifier(this, Montage, StartTime, EndTime);
-
-									//@TODO: Temp hack to keep track of the AnimNotifyState each modifier is created from.
-									if (RootMotionModifiers.Num())
-									{
-										TSharedPtr<FRootMotionModifier> Last = RootMotionModifiers.Last();
-										if(Last.IsValid() && Last->Animation.Get() == Montage && Last->StartTime == StartTime && Last->EndTime == EndTime)
-										{
-											Last->AnimNotifyState = MotionWarpingNotify;
-											Last->AnimNotifyState->OnWarpBegin(this, Last->Animation.Get(), Last->StartTime, Last->EndTime);
-										}
-									}
+									MotionWarpingNotify->OnBecomeRelevant(this, Montage, StartTime, EndTime);
 								}
 							}
 						}
@@ -274,21 +296,8 @@ void UMotionWarpingComponent::Update()
 	{
 		for (TSharedPtr<FRootMotionModifier>& Modifier : RootMotionModifiers)
 		{
-			const UAnimNotifyState_MotionWarping* AnimNotify = Modifier->AnimNotifyState.Get();
-			if (AnimNotify)
-			{
-				AnimNotify->OnWarpPreUpdate(this, Modifier->Animation.Get(), Modifier->StartTime, Modifier->EndTime);
-			}
-
+			//@TODO: Replace with the signature that doesn't receive the comp via param
 			Modifier->Update(*this);
-
-			if(AnimNotify)
-			{
-				if(Modifier->GetState() == ERootMotionModifierState::Disabled || Modifier->GetState() == ERootMotionModifierState::MarkedForRemoval)
-				{
-					AnimNotify->OnWarpEnd(this, Modifier->Animation.Get(), Modifier->StartTime, Modifier->EndTime);
-				}
-			}
 		}
 
 		// Remove the modifiers that has been marked for removal
@@ -327,6 +336,7 @@ FTransform UMotionWarpingComponent::ProcessRootMotionPreConvertToWorld(const FTr
 	{
 		if (RootMotionModifier->GetState() == ERootMotionModifierState::Active && RootMotionModifier->bInLocalSpace)
 		{
+			//@TODO: Replace with the signature that doesn't receive the comp via param
 			FinalRootMotion = RootMotionModifier->ProcessRootMotion(*this, FinalRootMotion, DeltaSeconds);
 		}
 	}
@@ -350,6 +360,7 @@ FTransform UMotionWarpingComponent::ProcessRootMotionPostConvertToWorld(const FT
 	{
 		if (RootMotionModifier->GetState() == ERootMotionModifierState::Active && !RootMotionModifier->bInLocalSpace)
 		{
+			//@TODO: Replace with the signature that doesn't receive the comp via param
 			FinalRootMotion = RootMotionModifier->ProcessRootMotion(*this, FinalRootMotion, DeltaSeconds);
 		}
 	}
@@ -400,4 +411,32 @@ void UMotionWarpingComponent::AddOrUpdateSyncPoint(FName Name, const FMotionWarp
 int32 UMotionWarpingComponent::RemoveSyncPoint(FName Name)
 {
 	return SyncPoints.Remove(Name);
+}
+
+bool UMotionWarpingComponent::SetRootMotionModifierDelegatesByHandle(FRootMotionModifierHandle Handle, FOnRootMotionModifierDelegate OnActivate, FOnRootMotionModifierDelegate OnUpdate, FOnRootMotionModifierDelegate OnDeactivate)
+{
+	TSharedPtr<FRootMotionModifier> Modifier = GetRootMotionModifierByHandle(Handle);
+	if (Modifier.IsValid())
+	{
+		Modifier->OnActivateDelegate = OnActivate;
+		Modifier->OnUpdateDelegate = OnUpdate;
+		Modifier->OnDeactivateDelegate = OnDeactivate;
+		return true;
+	}
+
+	return false;
+}
+
+bool UMotionWarpingComponent::GetAnimationAndTimeRangeForRootMotionModifierByHandle(FRootMotionModifierHandle Handle, UAnimSequenceBase*& OutAnimation, float& OutStartTime, float& OutEndTime) const
+{
+	TSharedPtr<FRootMotionModifier> Modifier = GetRootMotionModifierByHandle(Handle);
+	if (Modifier.IsValid())
+	{
+		OutAnimation = const_cast<UAnimSequenceBase*>(Modifier->GetAnimation());
+		OutStartTime = Modifier->StartTime;
+		OutEndTime = Modifier->EndTime;
+		return true;
+	}
+
+	return false;
 }
\ No newline at end of file
diff --git a/Engine/Plugins/Experimental/Animation/MotionWarping/Source/MotionWarping/Private/RootMotionModifier.cpp b/Engine/Plugins/Experimental/Animation/MotionWarping/Source/MotionWarping/Private/RootMotionModifier.cpp
--- a/Engine/Plugins/Experimental/Animation/MotionWarping/Source/MotionWarping/Private/RootMotionModifier.cpp
+++ b/Engine/Plugins/Experimental/Animation/MotionWarping/Source/MotionWarping/Private/RootMotionModifier.cpp
@@ -10,19 +10,47 @@
 #include "Animation/AnimInstance.h"
 #include "AnimNotifyState_MotionWarping.h"
 
+const FRootMotionModifierHandle FRootMotionModifierHandle::InvalidHandle;
+
+void FRootMotionModifierHandle::GenerateNewHandle()
+{
+	// Must be in C++ to avoid duplicate statics across execution units
+	static int32 GHandle = 1;
+	Handle = GHandle++;
+}
+
 // FRootMotionModifier
 ///////////////////////////////////////////////////////////////
 
-void FRootMotionModifier::Update(UMotionWarpingComponent& OwnerComp)
+void FRootMotionModifier::Initialize(UMotionWarpingComponent* OwnerComp)
+{
+	check(OwnerComp);
+	OwnerComponent = OwnerComp;
+	Handle.GenerateNewHandle();
+}
+
+ACharacter* FRootMotionModifier::GetCharacterOwner() const
 {
-	const FAnimMontageInstance* RootMotionMontageInstance = OwnerComp.GetCharacterOwner()->GetRootMotionAnimMontageInstance();
+	UMotionWarpingComponent* OwnerComp = GetOwnerComponent();
+	return OwnerComp ? OwnerComp->GetCharacterOwner() : nullptr;
+}
+
+void FRootMotionModifier::Update(UMotionWarpingComponent& OwnerComp_DEPRECATED)
+{
+	const ACharacter* CharacterOwner = GetCharacterOwner();
+	if(CharacterOwner == nullptr)
+	{
+		return;
+	}
+
+	const FAnimMontageInstance* RootMotionMontageInstance = CharacterOwner->GetRootMotionAnimMontageInstance();
 	const UAnimMontage* Montage = RootMotionMontageInstance ? ToRawPtr(RootMotionMontageInstance->Montage) : nullptr;
 
 	// Mark for removal if our animation is not relevant anymore
 	if (Montage == nullptr || Montage != Animation)
 	{
 		UE_LOG(LogMotionWarping, Verbose, TEXT("MotionWarping: Marking RootMotionModifier for removal. Reason: Animation is not valid. Char: %s Current Montage: %s. Window: Animation: %s [%f %f] [%f %f]"),
-			*GetNameSafe(OwnerComp.GetCharacterOwner()), *GetNameSafe(Montage), *GetNameSafe(Animation.Get()), StartTime, EndTime, PreviousPosition, CurrentPosition);
+			*GetNameSafe(CharacterOwner), *GetNameSafe(Montage), *GetNameSafe(Animation.Get()), StartTime, EndTime, PreviousPosition, CurrentPosition);
 
 		SetState(ERootMotionModifierState::MarkedForRemoval);
 		return;
@@ -37,7 +65,7 @@ void FRootMotionModifier::Update(UMotionWarpingComponent& OwnerComp)
 	if (PreviousPosition >= EndTime)
 	{
 		UE_LOG(LogMotionWarping, Verbose, TEXT("MotionWarping: Marking RootMotionModifier for removal. Reason: Window has ended. Char: %s Animation: %s [%f %f] [%f %f]"),
-			*GetNameSafe(OwnerComp.GetCharacterOwner()), *GetNameSafe(Animation.Get()), StartTime, EndTime, PreviousPosition, CurrentPosition);
+			*GetNameSafe(CharacterOwner), *GetNameSafe(Animation.Get()), StartTime, EndTime, PreviousPosition, CurrentPosition);
 
 		SetState(ERootMotionModifierState::MarkedForRemoval);
 		return;
@@ -52,6 +80,14 @@ void FRootMotionModifier::Update(UMotionWarpingComponent& OwnerComp)
 			SetState(ERootMotionModifierState::Active);
 		}
 	}
+
+	if (State == ERootMotionModifierState::Active)
+	{
+		if(UMotionWarpingComponent* OwnerComp = GetOwnerComponent())
+		{
+			OnUpdateDelegate.ExecuteIfBound(OwnerComp, Handle);
+		}
+	}
 }
 
 void FRootMotionModifier::SetState(ERootMotionModifierState NewState)
@@ -66,42 +102,122 @@ void FRootMotionModifier::SetState(ERootMotionModifierState NewState)
 	}
 }
 
+void FRootMotionModifier::OnStateChanged(ERootMotionModifierState LastState)
+{
+	if (UMotionWarpingComponent* OwnerComp = GetOwnerComponent())
+	{
+		if (LastState != ERootMotionModifierState::Active && State == ERootMotionModifierState::Active)
+		{
+			OnActivateDelegate.ExecuteIfBound(OwnerComp, Handle);
+		}
+		else if (LastState == ERootMotionModifierState::Active && (State == ERootMotionModifierState::Disabled || State == ERootMotionModifierState::MarkedForRemoval))
+		{
+			OnDeactivateDelegate.ExecuteIfBound(OwnerComp, Handle);
+		}
+	}
+}
+
 // FRootMotionModifier_Warp
 ///////////////////////////////////////////////////////////////
 
-void FRootMotionModifier_Warp::Update(UMotionWarpingComponent& OwnerComp)
+void FRootMotionModifier_Warp::Update(UMotionWarpingComponent& OwnerComp_DEPRECATED)
 {
 	// Update playback times and state
-	FRootMotionModifier::Update(OwnerComp);
+	FRootMotionModifier::Update(OwnerComp_DEPRECATED);
 
 	// Cache sync point transform and trigger OnSyncPointChanged if needed
 	if (GetState() == ERootMotionModifierState::Active)
 	{
-		const FMotionWarpingSyncPoint* SyncPointPtr = OwnerComp.FindSyncPoint(SyncPointName);
-
+		UMotionWarpingComponent* OwnerComp = GetOwnerComponent();
+		const FMotionWarpingSyncPoint* SyncPointPtr = OwnerComp ? OwnerComp->FindSyncPoint(SyncPointName) : nullptr;
+		
 		// Disable if there is no sync point for us
 		if (SyncPointPtr == nullptr)
 		{
 			UE_LOG(LogMotionWarping, Verbose, TEXT("MotionWarping: Marking RootMotionModifier as Disabled. Reason: Invalid Sync Point (%s). Char: %s Animation: %s [%f %f] [%f %f]"),
-				*SyncPointName.ToString(), *GetNameSafe(OwnerComp.GetCharacterOwner()), *GetNameSafe(Animation.Get()), StartTime, EndTime, PreviousPosition, CurrentPosition);
+				*SyncPointName.ToString(), *GetNameSafe(OwnerComp->GetOwner()), *GetNameSafe(Animation.Get()), StartTime, EndTime, PreviousPosition, CurrentPosition);
 
 			SetState(ERootMotionModifierState::Disabled);
 			return;
 		}
 
-		if (CachedSyncPoint != *SyncPointPtr)
+		// Get the warp point sent by the game
+		FTransform WarpPointTransformGame = FTransform(SyncPointPtr->GetRotation(), SyncPointPtr->GetLocation());
+
+		// Initialize our target transform (where the root should end at the end of the window) with the warp point sent by the game
+		FTransform TargetTransform = WarpPointTransformGame;
+
+		// Check if a warp point is defined in the animation. If so, we need to extract it and offset the target transform 
+		// the same amount the root bone is offset from the warp point in the animation
+		if(WarpPointAnimProvider != EWarpPointAnimProvider::None)
+		{
+			if(!CachedOffsetFromWarpPoint.IsSet())
+			{
+				if(const ACharacter* CharacterOwner = GetCharacterOwner())
+				{
+					// Inverse of mesh's relative rotation. Used to convert root and warp point in the animation from Y forward to X forward
+					const FTransform MeshCompRelativeRotInverse = FTransform(CharacterOwner->GetBaseRotationOffset().Inverse());
+
+					FTransform RootTransform = FTransform::Identity;
+					FTransform WarpPointTransform = FTransform::Identity;
+					if (WarpPointAnimProvider == EWarpPointAnimProvider::Static)
+					{
+						RootTransform = MeshCompRelativeRotInverse * UMotionWarpingUtilities::ExtractRootTransformFromAnimation(GetAnimation(), EndTime);
+						WarpPointTransform = MeshCompRelativeRotInverse * WarpPointAnimTransform;
+					}
+					else if (WarpPointAnimProvider == EWarpPointAnimProvider::Bone)
+					{
+						if (const USkeletalMeshComponent* Mesh = CharacterOwner->GetMesh())
+						{
+							if (const UAnimInstance* AnimInstance = Mesh->GetAnimInstance())
+							{
+								const FBoneContainer& FullBoneContainer = AnimInstance->GetRequiredBones();
+								const int32 BoneIndex = FullBoneContainer.GetPoseBoneIndexForBoneName(WarpPointAnimBoneName);
+								if (ensure(BoneIndex != INDEX_NONE))
+								{
+									TArray<FBoneIndexType> RequiredBoneIndexArray = { 0, (FBoneIndexType)BoneIndex };
+									FullBoneContainer.GetReferenceSkeleton().EnsureParentsExistAndSort(RequiredBoneIndexArray);
+
+									FBoneContainer LimitedBoneContainer(RequiredBoneIndexArray, FCurveEvaluationOption(false), *Animation->GetSkeleton());
+
+									FCSPose<FCompactPose> Pose;
+									UMotionWarpingUtilities::ExtractComponentSpacePose(GetAnimation(), LimitedBoneContainer, EndTime, false, Pose);
+
+									RootTransform = MeshCompRelativeRotInverse * Pose.GetComponentSpaceTransform(FCompactPoseBoneIndex(0));
+									WarpPointTransform = MeshCompRelativeRotInverse * Pose.GetComponentSpaceTransform(FCompactPoseBoneIndex(1));
+								}
+							}
+						}
+					}
+
+					// Calculate and cache the offset between the root and the warp point in the animation
+					CachedOffsetFromWarpPoint = WarpPointTransform.GetRelativeTransform(RootTransform);
+				}
+			}
+			
+			// Update Target Transform based on the offset between the root and the warp point in the animation
+			TargetTransform = CachedOffsetFromWarpPoint.GetValue() * WarpPointTransformGame;
+		}
+
+		if(!CachedTargetTransform.Equals(TargetTransform))
 		{
-			CachedSyncPoint = *SyncPointPtr;
+			CachedTargetTransform = TargetTransform;
 
-			OnSyncPointChanged(OwnerComp);
+			//DEPRECATED: Will be removed shortly
+			CachedSyncPoint = FMotionWarpingSyncPoint(CachedTargetTransform);
+
+			OnTargetTransformChanged();
 		}
 	}
 }
 
-FTransform FRootMotionModifier_Warp::ProcessRootMotion(UMotionWarpingComponent& OwnerComp, const FTransform& InRootMotion, float DeltaSeconds)
+FTransform FRootMotionModifier_Warp::ProcessRootMotion(UMotionWarpingComponent& OwnerComp_DEPRECATED, const FTransform& InRootMotion, float DeltaSeconds)
 {
-	const ACharacter* CharacterOwner = OwnerComp.GetCharacterOwner();
-	check(CharacterOwner);
+	const ACharacter* CharacterOwner = GetCharacterOwner();
+	if(CharacterOwner == nullptr)
+	{
+		return InRootMotion;
+	}
 
 	const FTransform& CharacterTransform = CharacterOwner->GetActorTransform();
 
@@ -116,7 +232,7 @@ FTransform FRootMotionModifier_Warp::ProcessRootMotion(UMotionWarpingComponent&
 		const FTransform RootMotionDelta = UMotionWarpingUtilities::ExtractRootMotionFromAnimation(Animation.Get(), PreviousPosition, FMath::Min(CurrentPosition, EndTime));
 
 		const float HorizontalDelta = RootMotionDelta.GetTranslation().Size2D();
-		const float HorizontalTarget = FVector::Dist2D(CharacterTransform.GetLocation(), CachedSyncPoint.GetLocation());
+		const float HorizontalTarget = FVector::Dist2D(CharacterTransform.GetLocation(), GetTargetLocation());
 		const float HorizontalOriginal = RootMotionTotal.GetTranslation().Size2D();
 		const float HorizontalTranslationWarped = HorizontalOriginal != 0.f ? ((HorizontalDelta * HorizontalTarget) / HorizontalOriginal) : 0.f;
 
@@ -124,18 +240,18 @@ FTransform FRootMotionModifier_Warp::ProcessRootMotion(UMotionWarpingComponent&
 		{
 			const FTransform MeshRelativeTransform = FTransform(CharacterOwner->GetBaseRotationOffset(), CharacterOwner->GetBaseTranslationOffset());
 			const FTransform MeshTransform = MeshRelativeTransform * CharacterOwner->GetActorTransform();
-			DeltaTranslation = MeshTransform.InverseTransformPositionNoScale(CachedSyncPoint.GetLocation()).GetSafeNormal2D() * HorizontalTranslationWarped;
+			DeltaTranslation = MeshTransform.InverseTransformPositionNoScale(GetTargetLocation()).GetSafeNormal2D() * HorizontalTranslationWarped;
 		}
 		else
 		{
-			DeltaTranslation = (CachedSyncPoint.GetLocation() - CharacterTransform.GetLocation()).GetSafeNormal2D() * HorizontalTranslationWarped;
+			DeltaTranslation = (GetTargetLocation() - CharacterTransform.GetLocation()).GetSafeNormal2D() * HorizontalTranslationWarped;
 		}
 
 		if (!bIgnoreZAxis)
 		{
 			const FVector CapsuleBottomLocation = (CharacterOwner->GetActorLocation() - FVector::UpVector * CharacterOwner->GetSimpleCollisionHalfHeight());
 			const float VerticalDelta = RootMotionDelta.GetTranslation().Z;
-			const float VerticalTarget = CachedSyncPoint.GetLocation().Z - CapsuleBottomLocation.Z;
+			const float VerticalTarget = GetTargetLocation().Z - CapsuleBottomLocation.Z;
 			const float VerticalOriginal = RootMotionTotal.GetTranslation().Z;
 			const float VerticalTranslationWarped = VerticalOriginal != 0.f ? ((VerticalDelta * VerticalTarget) / VerticalOriginal) : 0.f;
 
@@ -147,7 +263,7 @@ FTransform FRootMotionModifier_Warp::ProcessRootMotion(UMotionWarpingComponent&
 
 	if (bWarpRotation)
 	{
-		const FQuat WarpedRotation = WarpRotation(OwnerComp, InRootMotion, RootMotionTotal, DeltaSeconds);
+		const FQuat WarpedRotation = WarpRotation(InRootMotion, RootMotionTotal, DeltaSeconds);
 		FinalRootMotion.SetRotation(WarpedRotation);
 	}
 
@@ -156,12 +272,12 @@ FTransform FRootMotionModifier_Warp::ProcessRootMotion(UMotionWarpingComponent&
 	const int32 DebugLevel = FMotionWarpingCVars::CVarMotionWarpingDebug.GetValueOnGameThread();
 	if (DebugLevel > 0)
 	{
-		PrintLog(OwnerComp, TEXT("FRootMotionModifier_Simple"), InRootMotion, FinalRootMotion);
+		PrintLog(TEXT("FRootMotionModifier_Simple"), InRootMotion, FinalRootMotion);
 
 		if (DebugLevel >= 2)
 		{
 			const float DrawDebugDuration = FMotionWarpingCVars::CVarMotionWarpingDrawDebugDuration.GetValueOnGameThread();
-			DrawDebugCoordinateSystem(OwnerComp.GetWorld(), CachedSyncPoint.GetLocation(), CachedSyncPoint.Rotator(), 50.f, false, DrawDebugDuration, 0, 1.f);
+			DrawDebugCoordinateSystem(CharacterOwner->GetWorld(), GetTargetLocation(), GetTargetRotator(), 50.f, false, DrawDebugDuration, 0, 1.f);
 		}
 	}
 #endif
@@ -169,28 +285,36 @@ FTransform FRootMotionModifier_Warp::ProcessRootMotion(UMotionWarpingComponent&
 	return FinalRootMotion;
 }
 
-FQuat FRootMotionModifier_Warp::GetTargetRotation(UMotionWarpingComponent& OwnerComp) const
+FQuat FRootMotionModifier_Warp::GetTargetRotation() const
 {
-	const FTransform& CharacterTransform = OwnerComp.GetCharacterOwner()->GetActorTransform();
-
 	if (RotationType == EMotionWarpRotationType::Default)
 	{
-		return CachedSyncPoint.GetRotation();
+		return CachedTargetTransform.GetRotation();
 	}
 	else if (RotationType == EMotionWarpRotationType::Facing)
 	{
-		const FVector ToSyncPoint = (CachedSyncPoint.GetLocation() - CharacterTransform.GetLocation()).GetSafeNormal2D();
-		return FRotationMatrix::MakeFromXZ(ToSyncPoint, FVector::UpVector).ToQuat();
+		if(const ACharacter* CharacterOwner = GetCharacterOwner())
+		{
+			const FTransform& CharacterTransform = CharacterOwner->GetActorTransform();
+			const FVector ToSyncPoint = (CachedTargetTransform.GetLocation() - CharacterTransform.GetLocation()).GetSafeNormal2D();
+			return FRotationMatrix::MakeFromXZ(ToSyncPoint, FVector::UpVector).ToQuat();
+		}
 	}
 
 	return FQuat::Identity;
 }
 
-FQuat FRootMotionModifier_Warp::WarpRotation(UMotionWarpingComponent& OwnerComp, const FTransform& RootMotionDelta, const FTransform& RootMotionTotal, float DeltaSeconds)
+FQuat FRootMotionModifier_Warp::WarpRotation(const FTransform& RootMotionDelta, const FTransform& RootMotionTotal, float DeltaSeconds)
 {
-	const FTransform& CharacterTransform = OwnerComp.GetCharacterOwner()->GetActorTransform();
+	const ACharacter* CharacterOwner = GetCharacterOwner();
+	if (CharacterOwner == nullptr)
+	{
+		return FQuat::Identity;
+	}
+
+	const FTransform& CharacterTransform = CharacterOwner->GetActorTransform();
 	const FQuat CurrentRotation = CharacterTransform.GetRotation();
-	const FQuat TargetRotation = GetTargetRotation(OwnerComp);
+	const FQuat TargetRotation = GetTargetRotation();
 	const float TimeRemaining = (EndTime - PreviousPosition) * WarpRotationTimeMultiplier;
 	const FQuat RemainingRootRotationInWorld = RootMotionTotal.GetRotation();
 	const FQuat CurrentPlusRemainingRootMotion = RemainingRootRotationInWorld * CurrentRotation;
@@ -202,34 +326,36 @@ FQuat FRootMotionModifier_Warp::WarpRotation(UMotionWarpingComponent& OwnerComp,
 }
 
 #if !(UE_BUILD_SHIPPING || UE_BUILD_TEST)
-void FRootMotionModifier_Warp::PrintLog(const UMotionWarpingComponent& OwnerComp, const FString& Name, const FTransform& OriginalRootMotion, const FTransform& WarpedRootMotion) const
+void FRootMotionModifier_Warp::PrintLog(const FString& Name, const FTransform& OriginalRootMotion, const FTransform& WarpedRootMotion) const
 {
-	const ACharacter* CharacterOwner = OwnerComp.GetCharacterOwner();
-	check(CharacterOwner);
-
-	const FVector CurrentLocation = (CharacterOwner->GetActorLocation() - FVector::UpVector * CharacterOwner->GetSimpleCollisionHalfHeight());
-	const FVector CurrentToTarget = (CachedSyncPoint.GetLocation() - CurrentLocation).GetSafeNormal2D();
-	const FVector FutureLocation = CurrentLocation + (bInLocalSpace ? (CharacterOwner->GetMesh()->ConvertLocalRootMotionToWorld(WarpedRootMotion)).GetTranslation() : WarpedRootMotion.GetTranslation());
-	const FRotator CurrentRotation = CharacterOwner->GetActorRotation();
-	const FRotator FutureRotation = (WarpedRootMotion.GetRotation() * CharacterOwner->GetActorQuat()).Rotator();
-	const float Dot = FVector::DotProduct(CharacterOwner->GetActorForwardVector(), CurrentToTarget);
-	const float CurrentDist2D = FVector::Dist2D(CachedSyncPoint.GetLocation(), CurrentLocation);
-	const float FutureDist2D = FVector::Dist2D(CachedSyncPoint.GetLocation(), FutureLocation);
-	const float DeltaSeconds = CharacterOwner->GetWorld()->GetDeltaSeconds();
-	const float Speed = WarpedRootMotion.GetTranslation().Size() / DeltaSeconds;
-	const float EndTimeOffset = CurrentPosition - EndTime;
-
-	UE_LOG(LogMotionWarping, Log, TEXT("MotionWarping: %s. NetMode: %d Char: %s Anim: %s Window [%f %f][%f %f] DeltaTime: %f WorldTime: %f EndTimeOffset: %f Dist2D: %f FutureDist2D: %f Dot: %f OriginalMotionDelta: %s (%f) FinalMotionDelta: %s (%f) Speed: %f Loc: %s FutureLoc: %s Rot: %s FutureRot: %s"),
-		*Name, (int32)CharacterOwner->GetWorld()->GetNetMode(), *GetNameSafe(CharacterOwner), *GetNameSafe(Animation.Get()), StartTime, EndTime, PreviousPosition, CurrentPosition, DeltaSeconds, CharacterOwner->GetWorld()->GetTimeSeconds(), EndTimeOffset, CurrentDist2D, FutureDist2D, Dot,
-		*OriginalRootMotion.GetTranslation().ToString(), OriginalRootMotion.GetTranslation().Size(), *WarpedRootMotion.GetTranslation().ToString(), WarpedRootMotion.GetTranslation().Size(), Speed,
-		*CurrentLocation.ToString(), *FutureLocation.ToString(), *CurrentRotation.ToCompactString(), *FutureRotation.ToCompactString());
+	if(const ACharacter* CharacterOwner = GetCharacterOwner())
+	{
+		const FVector CurrentLocation = (CharacterOwner->GetActorLocation() - FVector::UpVector * CharacterOwner->GetSimpleCollisionHalfHeight());
+		const FVector CurrentToTarget = (GetTargetLocation() - CurrentLocation).GetSafeNormal2D();
+		const FVector FutureLocation = CurrentLocation + (bInLocalSpace ? (CharacterOwner->GetMesh()->ConvertLocalRootMotionToWorld(WarpedRootMotion)).GetTranslation() : WarpedRootMotion.GetTranslation());
+		const FRotator CurrentRotation = CharacterOwner->GetActorRotation();
+		const FRotator FutureRotation = (WarpedRootMotion.GetRotation() * CharacterOwner->GetActorQuat()).Rotator();
+		const float Dot = FVector::DotProduct(CharacterOwner->GetActorForwardVector(), CurrentToTarget);
+		const float CurrentDist2D = FVector::Dist2D(GetTargetLocation(), CurrentLocation);
+		const float FutureDist2D = FVector::Dist2D(GetTargetLocation(), FutureLocation);
+		const float DeltaSeconds = CharacterOwner->GetWorld()->GetDeltaSeconds();
+		const float Speed = WarpedRootMotion.GetTranslation().Size() / DeltaSeconds;
+		const float EndTimeOffset = CurrentPosition - EndTime;
+
+		UE_LOG(LogMotionWarping, Log, TEXT("MotionWarping: %s. NetMode: %d Char: %s Anim: %s Win: [%f %f][%f %f] DT: %f WT: %f EndTimeOffset: %f Dist2D: %f FutureDist2D: %f Dot: %f OriginalMotionDelta: %s (%f) FinalMotionDelta: %s (%f) Speed: %f Loc: %s FutureLoc: %s Rot: %s FutureRot: %s"),
+			*Name, (int32)CharacterOwner->GetWorld()->GetNetMode(), *GetNameSafe(CharacterOwner), *GetNameSafe(Animation.Get()), StartTime, EndTime, PreviousPosition, CurrentPosition, DeltaSeconds, CharacterOwner->GetWorld()->GetTimeSeconds(), EndTimeOffset, CurrentDist2D, FutureDist2D, Dot,
+			*OriginalRootMotion.GetTranslation().ToString(), OriginalRootMotion.GetTranslation().Size(), *WarpedRootMotion.GetTranslation().ToString(), WarpedRootMotion.GetTranslation().Size(), Speed,
+			*CurrentLocation.ToString(), *FutureLocation.ToString(), *CurrentRotation.ToCompactString(), *FutureRotation.ToCompactString());
+	}
 }
 #endif
 
 // URootMotionModifierConfig_Warp
 ///////////////////////////////////////////////////////////////
 
-void URootMotionModifierConfig_Warp::AddRootMotionModifierSimpleWarp(UMotionWarpingComponent* InMotionWarpingComp, const UAnimSequenceBase* InAnimation, float InStartTime, float InEndTime, FName InSyncPointName, bool bInWarpTranslation, bool bInIgnoreZAxis, bool bInWarpRotation, EMotionWarpRotationType InRotationType, float InWarpRotationTimeMultiplier)
+FRootMotionModifierHandle URootMotionModifierConfig_Warp::AddRootMotionModifierSimpleWarp(UMotionWarpingComponent* InMotionWarpingComp, const UAnimSequenceBase* InAnimation, float InStartTime, float InEndTime,
+	FName InSyncPointName, EWarpPointAnimProvider InWarpPointAnimProvider, FTransform InWarpPointAnimTransform, FName InWarpPointAnimBoneName,
+	bool bInWarpTranslation, bool bInIgnoreZAxis, bool bInWarpRotation, EMotionWarpRotationType InRotationType, float InWarpRotationTimeMultiplier)
 {
 	if (ensureAlways(InMotionWarpingComp))
 	{
@@ -238,19 +364,24 @@ void URootMotionModifierConfig_Warp::AddRootMotionModifierSimpleWarp(UMotionWarp
 		NewModifier->StartTime = InStartTime;
 		NewModifier->EndTime = InEndTime;
 		NewModifier->SyncPointName = InSyncPointName;
+		NewModifier->WarpPointAnimProvider = InWarpPointAnimProvider;
+		NewModifier->WarpPointAnimTransform = InWarpPointAnimTransform;
+		NewModifier->WarpPointAnimBoneName = InWarpPointAnimBoneName;
 		NewModifier->bWarpTranslation = bInWarpTranslation;
 		NewModifier->bIgnoreZAxis = bInIgnoreZAxis;
 		NewModifier->bWarpRotation = bInWarpRotation;
 		NewModifier->RotationType = InRotationType;
 		NewModifier->WarpRotationTimeMultiplier = InWarpRotationTimeMultiplier;
-		InMotionWarpingComp->AddRootMotionModifier(NewModifier);
+		return InMotionWarpingComp->AddRootMotionModifier(NewModifier);
 	}
+
+	return FRootMotionModifierHandle::InvalidHandle;
 }
 
 // URootMotionModifierConfig_Scale
 ///////////////////////////////////////////////////////////////
 
-void URootMotionModifierConfig_Scale::AddRootMotionModifierScale(UMotionWarpingComponent* InMotionWarpingComp, const UAnimSequenceBase* InAnimation, float InStartTime, float InEndTime, FVector InScale)
+FRootMotionModifierHandle URootMotionModifierConfig_Scale::AddRootMotionModifierScale(UMotionWarpingComponent* InMotionWarpingComp, const UAnimSequenceBase* InAnimation, float InStartTime, float InEndTime, FVector InScale)
 {
 	if (ensureAlways(InMotionWarpingComp))
 	{
@@ -259,6 +390,8 @@ void URootMotionModifierConfig_Scale::AddRootMotionModifierScale(UMotionWarpingC
 		NewModifier->StartTime = InStartTime;
 		NewModifier->EndTime = InEndTime;
 		NewModifier->Scale = InScale;
-		InMotionWarpingComp->AddRootMotionModifier(NewModifier);
+		return InMotionWarpingComp->AddRootMotionModifier(NewModifier);
 	}
-}
\ No newline at end of file
+
+	return FRootMotionModifierHandle::InvalidHandle;
+}
diff --git a/Engine/Plugins/Experimental/Animation/MotionWarping/Source/MotionWarping/Private/RootMotionModifier_AdjustmentBlendWarp.cpp b/Engine/Plugins/Experimental/Animation/MotionWarping/Source/MotionWarping/Private/RootMotionModifier_AdjustmentBlendWarp.cpp
--- a/Engine/Plugins/Experimental/Animation/MotionWarping/Source/MotionWarping/Private/RootMotionModifier_AdjustmentBlendWarp.cpp
+++ b/Engine/Plugins/Experimental/Animation/MotionWarping/Source/MotionWarping/Private/RootMotionModifier_AdjustmentBlendWarp.cpp
@@ -51,26 +51,29 @@ void FRootMotionModifier_AdjustmentBlendWarp::ExtractBoneTransformAtFrame(FTrans
 	FAnimationUtils::ExtractTransformForFrameFromTrack(Result.AnimationTracks[TrackIndex], Frame, OutTransform);
 }
 
-void FRootMotionModifier_AdjustmentBlendWarp::OnSyncPointChanged(UMotionWarpingComponent& OwnerComp)
+void FRootMotionModifier_AdjustmentBlendWarp::OnTargetTransformChanged()
 {
-	const ACharacter* CharacterOwner = OwnerComp.GetCharacterOwner();
-	const USkeletalMeshComponent* SkelMeshComp = CharacterOwner->GetMesh();
+	const ACharacter* CharacterOwner = GetCharacterOwner();
+	const USkeletalMeshComponent* SkelMeshComp = CharacterOwner ? CharacterOwner->GetMesh() : nullptr;
 
-	ActualStartTime = PreviousPosition;
-	CachedRootMotion = UMotionWarpingUtilities::ExtractRootMotionFromAnimation(Animation.Get(), ActualStartTime, EndTime);
-	CachedMeshRelativeTransform = SkelMeshComp->GetRelativeTransform();
-	CachedMeshTransform = SkelMeshComp->GetComponentTransform();
+	if(SkelMeshComp)
+	{
+		ActualStartTime = PreviousPosition;
+		CachedRootMotion = UMotionWarpingUtilities::ExtractRootMotionFromAnimation(Animation.Get(), ActualStartTime, EndTime);
+		CachedMeshRelativeTransform = SkelMeshComp->GetRelativeTransform();
+		CachedMeshTransform = SkelMeshComp->GetComponentTransform();
 
-	Result.AnimationTracks.Reset();
-	Result.TrackNames.Reset();
+		Result.AnimationTracks.Reset();
+		Result.TrackNames.Reset();
+	}
 }
 
-FTransform FRootMotionModifier_AdjustmentBlendWarp::ProcessRootMotion(UMotionWarpingComponent& OwnerComp, const FTransform& InRootMotion, float DeltaSeconds)
+FTransform FRootMotionModifier_AdjustmentBlendWarp::ProcessRootMotion(UMotionWarpingComponent& OwnerComp_DEPRECATED, const FTransform& InRootMotion, float DeltaSeconds)
 {
 	// If warped tracks has not been generated yet, do it now
 	if (Result.GetNum() == 0)
 	{
-		PrecomputeWarpedTracks(OwnerComp);
+		PrecomputeWarpedTracks();
 	}
 
 	// Extract root motion from warped tracks
@@ -80,12 +83,12 @@ FTransform FRootMotionModifier_AdjustmentBlendWarp::ProcessRootMotion(UMotionWar
 #if !(UE_BUILD_SHIPPING || UE_BUILD_TEST)
 	if (FMotionWarpingCVars::CVarMotionWarpingDebug.GetValueOnGameThread() > 0)
 	{
-		PrintLog(OwnerComp, TEXT("FRootMotionModifier_AdjustmentBlend"), InRootMotion, FinalRootMotion);
+		PrintLog(TEXT("FRootMotionModifier_AdjustmentBlend"), InRootMotion, FinalRootMotion);
 
 		if(FMotionWarpingCVars::CVarMotionWarpingDebug.GetValueOnGameThread() >= 2)
 		{
 			const float DrawDebugDuration = FMotionWarpingCVars::CVarMotionWarpingDrawDebugDuration.GetValueOnGameThread();
-			DrawDebugWarpedTracks(OwnerComp, DrawDebugDuration);
+			DrawDebugWarpedTracks(DrawDebugDuration);
 		}
 	}
 #endif
@@ -93,13 +96,18 @@ FTransform FRootMotionModifier_AdjustmentBlendWarp::ProcessRootMotion(UMotionWar
 	return FinalRootMotion;
 }
 
-void FRootMotionModifier_AdjustmentBlendWarp::PrecomputeWarpedTracks(UMotionWarpingComponent& OwnerComp)
+void FRootMotionModifier_AdjustmentBlendWarp::PrecomputeWarpedTracks()
 {
 	SCOPE_CYCLE_COUNTER(STAT_MotionWarping_PrecomputeWarpedTracks);
 
 	// First, extract pose at the end of the window for the bones we are going to warp
 
-	const ACharacter* CharacterOwner = OwnerComp.GetCharacterOwner();
+	const ACharacter* CharacterOwner = GetCharacterOwner();
+	if(CharacterOwner == nullptr)
+	{
+		return;
+	}
+
 	const FBoneContainer& BoneContainer = CharacterOwner->GetMesh()->GetAnimInstance()->GetRequiredBones();
 
 	// Init FBoneContainer with only the bones that we are interested in
@@ -131,7 +139,7 @@ void FRootMotionModifier_AdjustmentBlendWarp::PrecomputeWarpedTracks(UMotionWarp
 	// Second, calculate additive pose
 
 	//Calculate additive translation for root bone
-	FVector RootTargetLocation = CachedMeshTransform.InverseTransformPositionNoScale(CachedSyncPoint.GetLocation());
+	FVector RootTargetLocation = CachedMeshTransform.InverseTransformPositionNoScale(GetTargetLocation());
 
 	FVector RootTotalAdditiveTranslation = FVector::ZeroVector;
 	if (bWarpTranslation)
@@ -148,7 +156,7 @@ void FRootMotionModifier_AdjustmentBlendWarp::PrecomputeWarpedTracks(UMotionWarp
 	FQuat RootTotalAdditiveRotation = FQuat::Identity;
 	if (bWarpRotation)
 	{
-		const FQuat TargetRotation = GetTargetRotation(OwnerComp);
+		const FQuat TargetRotation = GetTargetRotation();
 		const FQuat OriginalRotation = CachedMeshRelativeTransform.GetRotation().Inverse() * (CachedRootMotion * CachedMeshTransform).GetRotation();
 		RootTotalAdditiveRotation = FQuat::FindBetweenNormals(OriginalRotation.GetForwardVector(), TargetRotation.GetForwardVector());
 	}
@@ -341,9 +349,10 @@ FTransform FRootMotionModifier_AdjustmentBlendWarp::ExtractWarpedRootMotion() co
 }
 
 #if !(UE_BUILD_SHIPPING || UE_BUILD_TEST)
-void FRootMotionModifier_AdjustmentBlendWarp::DrawDebugWarpedTracks(UMotionWarpingComponent& OwnerComp, float DrawDuration) const
+void FRootMotionModifier_AdjustmentBlendWarp::DrawDebugWarpedTracks(float DrawDuration) const
 {
-	const UWorld* World = OwnerComp.GetWorld();
+	const ACharacter* CharacterOwner = GetCharacterOwner();
+	const UWorld* World = CharacterOwner ? CharacterOwner->GetWorld() : nullptr;
 	if (World && Result.GetNum() > 0 && PreviousPosition <= EndTime)
 	{
 		FTransform RootStartTransform;
@@ -369,10 +378,10 @@ void FRootMotionModifier_AdjustmentBlendWarp::DrawDebugWarpedTracks(UMotionWarpi
 			}
 		}
 
-		const FTransform& MeshTransform = OwnerComp.GetCharacterOwner()->GetMesh()->GetComponentTransform();
+		const FTransform& MeshTransform = CharacterOwner->GetMesh()->GetComponentTransform();
 		DrawDebugCoordinateSystem(World, CachedMeshTransform.GetLocation(), CachedMeshTransform.Rotator(), 20.f, false, DrawDuration, 0, 1.f);
 		DrawDebugCoordinateSystem(World, MeshTransform.GetLocation(), MeshTransform.Rotator(), 20.f, false, DrawDuration, 0, 1.f);
-		DrawDebugCoordinateSystem(World, CachedSyncPoint.GetLocation(), CachedSyncPoint.Rotator(), 50.f, false, DrawDuration, 0, 1.f);
+		DrawDebugCoordinateSystem(World, GetTargetLocation(), GetTargetRotator(), 50.f, false, DrawDuration, 0, 1.f);
 		DrawDebugCoordinateSystem(World, (CachedRootMotion * CachedMeshTransform).GetLocation(), (CachedRootMotion * CachedMeshTransform).Rotator(), 50.f, false, DrawDuration, 0, 1.f);
 	}
 }
@@ -426,7 +435,9 @@ void URootMotionModifierConfig_AdjustmentBlendWarp::GetIKBoneTransformAndAlpha(A
 // URootMotionModifierConfig_AdjustmentBlendWarp
 ///////////////////////////////////////////////////////////////
 
-void URootMotionModifierConfig_AdjustmentBlendWarp::AddRootMotionModifierAdjustmentBlendWarp(UMotionWarpingComponent* InMotionWarpingComp, const UAnimSequenceBase* InAnimation, float InStartTime, float InEndTime, FName InSyncPointName, bool bInWarpTranslation, bool bInIgnoreZAxis, bool bInWarpRotation, bool bInWarpIKBones, const TArray<FName>& InIKBones)
+FRootMotionModifierHandle URootMotionModifierConfig_AdjustmentBlendWarp::AddRootMotionModifierAdjustmentBlendWarp(UMotionWarpingComponent* InMotionWarpingComp, const UAnimSequenceBase* InAnimation, float InStartTime, float InEndTime, 
+	FName InSyncPointName, EWarpPointAnimProvider InWarpPointAnimProvider, FTransform InWarpPointAnimTransform, FName InWarpPointAnimBoneName,
+	bool bInWarpTranslation, bool bInIgnoreZAxis, bool bInWarpRotation, bool bInWarpIKBones, const TArray<FName>& InIKBones)
 {
 	if (ensureAlways(InMotionWarpingComp))
 	{
@@ -435,11 +446,16 @@ void URootMotionModifierConfig_AdjustmentBlendWarp::AddRootMotionModifierAdjustm
 		NewModifier->StartTime = InStartTime;
 		NewModifier->EndTime = InEndTime;
 		NewModifier->SyncPointName = InSyncPointName;
+		NewModifier->WarpPointAnimProvider = InWarpPointAnimProvider;
+		NewModifier->WarpPointAnimTransform = InWarpPointAnimTransform;
+		NewModifier->WarpPointAnimBoneName = InWarpPointAnimBoneName;
 		NewModifier->bWarpTranslation = bInWarpTranslation;
 		NewModifier->bIgnoreZAxis = bInIgnoreZAxis;
 		NewModifier->bWarpRotation = bInWarpRotation;
 		NewModifier->bWarpIKBones = bInWarpIKBones;
 		NewModifier->IKBones = InIKBones;
-		InMotionWarpingComp->AddRootMotionModifier(NewModifier);
+		return InMotionWarpingComp->AddRootMotionModifier(NewModifier);
 	}
+
+	return FRootMotionModifierHandle::InvalidHandle;
 }
\ No newline at end of file
diff --git a/Engine/Plugins/Experimental/Animation/MotionWarping/Source/MotionWarping/Private/RootMotionModifier_SkewWarp.cpp b/Engine/Plugins/Experimental/Animation/MotionWarping/Source/MotionWarping/Private/RootMotionModifier_SkewWarp.cpp
--- a/Engine/Plugins/Experimental/Animation/MotionWarping/Source/MotionWarping/Private/RootMotionModifier_SkewWarp.cpp
+++ b/Engine/Plugins/Experimental/Animation/MotionWarping/Source/MotionWarping/Private/RootMotionModifier_SkewWarp.cpp
@@ -7,10 +7,13 @@
 #include "Components/SkeletalMeshComponent.h"
 #include "MotionWarpingComponent.h"
 
-FTransform FRootMotionModifier_SkewWarp::ProcessRootMotion(UMotionWarpingComponent& OwnerComp, const FTransform& InRootMotion, float DeltaSeconds)
+FTransform FRootMotionModifier_SkewWarp::ProcessRootMotion(UMotionWarpingComponent& OwnerComp_DEPRECATED, const FTransform& InRootMotion, float DeltaSeconds)
 {
-	const ACharacter* CharacterOwner = OwnerComp.GetCharacterOwner(); 
-	check(CharacterOwner);
+	const ACharacter* CharacterOwner = GetCharacterOwner();
+	if(CharacterOwner == nullptr)
+	{
+		return InRootMotion;
+	}
 
 	FTransform FinalRootMotion = InRootMotion;
 
@@ -32,7 +35,7 @@ FTransform FRootMotionModifier_SkewWarp::ProcessRootMotion(UMotionWarpingCompone
 		const FVector CurrentLocation = CurrentTransform.GetLocation();
 		const FQuat CurrentRotation = CurrentTransform.GetRotation();
 		
-		FVector TargetLocation = CachedSyncPoint.GetLocation();
+		FVector TargetLocation = GetTargetLocation();
 		if(bIgnoreZAxis)
 		{
 			TargetLocation.Z = CurrentLocation.Z;
@@ -138,7 +141,7 @@ FTransform FRootMotionModifier_SkewWarp::ProcessRootMotion(UMotionWarpingCompone
 
 	if(bWarpRotation)
 	{
-		const FQuat WarpedRotation = WarpRotation(OwnerComp, InRootMotion, RootMotionTotal, DeltaSeconds);
+		const FQuat WarpedRotation = WarpRotation(InRootMotion, RootMotionTotal, DeltaSeconds);
 		FinalRootMotion.SetRotation(WarpedRotation);
 	}
 
@@ -147,12 +150,12 @@ FTransform FRootMotionModifier_SkewWarp::ProcessRootMotion(UMotionWarpingCompone
 	const int32 DebugLevel = FMotionWarpingCVars::CVarMotionWarpingDebug.GetValueOnGameThread();
 	if (DebugLevel > 0)
 	{
-		PrintLog(OwnerComp, TEXT("FRootMotionModifier_Skew"), InRootMotion, FinalRootMotion);
+		PrintLog(TEXT("FRootMotionModifier_Skew"), InRootMotion, FinalRootMotion);
 
 		if (DebugLevel >= 2)
 		{
 			const float DrawDebugDuration = FMotionWarpingCVars::CVarMotionWarpingDrawDebugDuration.GetValueOnGameThread();
-			DrawDebugCoordinateSystem(OwnerComp.GetWorld(), CachedSyncPoint.GetLocation(), CachedSyncPoint.Rotator(), 50.f, false, DrawDebugDuration, 0, 1.f);
+			DrawDebugCoordinateSystem(CharacterOwner->GetWorld(), GetTargetLocation(), GetTargetRotator(), 50.f, false, DrawDebugDuration, 0, 1.f);
 		}
 	}
 #endif
@@ -163,7 +166,9 @@ FTransform FRootMotionModifier_SkewWarp::ProcessRootMotion(UMotionWarpingCompone
 // URootMotionModifierConfig_SkewWarp
 ///////////////////////////////////////////////////////////////
 
-void URootMotionModifierConfig_SkewWarp::AddRootMotionModifierSkewWarp(UMotionWarpingComponent* InMotionWarpingComp, const UAnimSequenceBase* InAnimation, float InStartTime, float InEndTime, FName InSyncPointName, bool bInWarpTranslation, bool bInIgnoreZAxis, bool bInWarpRotation, EMotionWarpRotationType InRotationType, float InWarpRotationTimeMultiplier)
+FRootMotionModifierHandle URootMotionModifierConfig_SkewWarp::AddRootMotionModifierSkewWarp(UMotionWarpingComponent* InMotionWarpingComp, const UAnimSequenceBase* InAnimation, float InStartTime, float InEndTime, 
+	FName InSyncPointName, EWarpPointAnimProvider InWarpPointAnimProvider, FTransform InWarpPointAnimTransform, FName InWarpPointAnimBoneName,
+	bool bInWarpTranslation, bool bInIgnoreZAxis, bool bInWarpRotation, EMotionWarpRotationType InRotationType, float InWarpRotationTimeMultiplier)
 {
 	if (ensureAlways(InMotionWarpingComp))
 	{
@@ -172,11 +177,16 @@ void URootMotionModifierConfig_SkewWarp::AddRootMotionModifierSkewWarp(UMotionWa
 		NewModifier->StartTime = InStartTime;
 		NewModifier->EndTime = InEndTime;
 		NewModifier->SyncPointName = InSyncPointName;
+		NewModifier->WarpPointAnimProvider = InWarpPointAnimProvider;
+		NewModifier->WarpPointAnimTransform = InWarpPointAnimTransform;
+		NewModifier->WarpPointAnimBoneName = InWarpPointAnimBoneName;
 		NewModifier->bWarpTranslation = bInWarpTranslation;
 		NewModifier->bIgnoreZAxis = bInIgnoreZAxis;
 		NewModifier->bWarpRotation = bInWarpRotation;
 		NewModifier->RotationType = InRotationType;
 		NewModifier->WarpRotationTimeMultiplier = InWarpRotationTimeMultiplier;
-		InMotionWarpingComp->AddRootMotionModifier(NewModifier);
+		return InMotionWarpingComp->AddRootMotionModifier(NewModifier);
 	}
+
+	return FRootMotionModifierHandle::InvalidHandle;
 }
\ No newline at end of file
diff --git a/Engine/Plugins/Experimental/Animation/MotionWarping/Source/MotionWarping/Public/AnimNotifyState_MotionWarping.h b/Engine/Plugins/Experimental/Animation/MotionWarping/Source/MotionWarping/Public/AnimNotifyState_MotionWarping.h
--- a/Engine/Plugins/Experimental/Animation/MotionWarping/Source/MotionWarping/Public/AnimNotifyState_MotionWarping.h
+++ b/Engine/Plugins/Experimental/Animation/MotionWarping/Source/MotionWarping/Public/AnimNotifyState_MotionWarping.h
@@ -4,6 +4,7 @@
 
 #include "CoreMinimal.h"
 #include "Animation/AnimNotifies/AnimNotifyState.h"
+#include "RootMotionModifier.h"
 #include "AnimNotifyState_MotionWarping.generated.h"
 
 class UMotionWarpingComponent;
@@ -18,20 +19,35 @@ class MOTIONWARPING_API UAnimNotifyState_MotionWarping : public UAnimNotifyState
 
 public:
 
+	//@TODO: Prevent notify callbacks and add comments explaining why we don't use those here.
+
 	UPROPERTY(EditAnywhere, Instanced, BlueprintReadWrite, Category = "Config")
 	URootMotionModifierConfig* RootMotionModifierConfig;
 
 	UAnimNotifyState_MotionWarping(const FObjectInitializer& ObjectInitializer);
 
+	/** Called from the MotionWarpingComp when this notify becomes relevant. See: UMotionWarpingComponent::Update */
+	FRootMotionModifierHandle OnBecomeRelevant(UMotionWarpingComponent* MotionWarpingComp, const UAnimSequenceBase* Animation, float StartTime, float EndTime) const;
+
+	/** Creates a root motion modifier from the config class defined in the notify */
 	UFUNCTION(BlueprintNativeEvent, Category = "Motion Warping")
-	void AddRootMotionModifier(UMotionWarpingComponent* MotionWarpingComp, const UAnimSequenceBase* Animation, float StartTime, float EndTime) const;
+	FRootMotionModifierHandle AddRootMotionModifier(UMotionWarpingComponent* MotionWarpingComp, const UAnimSequenceBase* Animation, float StartTime, float EndTime) const;
+
+	UFUNCTION()
+	void OnRootMotionModifierActivate(UMotionWarpingComponent* MotionWarpingComp, const FRootMotionModifierHandle& Handle);
+
+	UFUNCTION()
+	void OnRootMotionModifierUpdate(UMotionWarpingComponent* MotionWarpingComp, const FRootMotionModifierHandle& Handle);
+
+	UFUNCTION()
+	void OnRootMotionModifierDeactivate(UMotionWarpingComponent* MotionWarpingComp, const FRootMotionModifierHandle& Handle);
 
 	UFUNCTION(BlueprintImplementableEvent, Category = "Motion Warping")
-	void OnWarpBegin(UMotionWarpingComponent* MotionWarpingComp, const UAnimSequenceBase* Animation, float StartTime, float EndTime) const;
+	void OnWarpBegin(UMotionWarpingComponent* MotionWarpingComp, const FRootMotionModifierHandle& Handle) const;
 
 	UFUNCTION(BlueprintImplementableEvent, Category = "Motion Warping")
-	void OnWarpPreUpdate(UMotionWarpingComponent* MotionWarpingComp, const UAnimSequenceBase* Animation, float StartTime, float EndTime) const;
+	void OnWarpUpdate(UMotionWarpingComponent* MotionWarpingComp, const FRootMotionModifierHandle& Handle) const;
 
 	UFUNCTION(BlueprintImplementableEvent, Category = "Motion Warping")
-	void OnWarpEnd(UMotionWarpingComponent* MotionWarpingComp, const UAnimSequenceBase* Animation, float StartTime, float EndTime) const;
+	void OnWarpEnd(UMotionWarpingComponent* MotionWarpingComp, const FRootMotionModifierHandle& Handle) const;
 };
\ No newline at end of file
diff --git a/Engine/Plugins/Experimental/Animation/MotionWarping/Source/MotionWarping/Public/MotionWarpingComponent.h b/Engine/Plugins/Experimental/Animation/MotionWarping/Source/MotionWarping/Public/MotionWarpingComponent.h
--- a/Engine/Plugins/Experimental/Animation/MotionWarping/Source/MotionWarping/Public/MotionWarpingComponent.h
+++ b/Engine/Plugins/Experimental/Animation/MotionWarping/Source/MotionWarping/Public/MotionWarpingComponent.h
@@ -65,6 +65,9 @@ public:
 	UFUNCTION(BlueprintCallable, Category = "Motion Warping")
 	static FTransform ExtractRootMotionFromAnimation(const UAnimSequenceBase* Animation, float StartTime, float EndTime);
 
+	/** Extract root bone transform at a given time */
+	static FTransform ExtractRootTransformFromAnimation(const UAnimSequenceBase* Animation, float Time);
+
 	/** @return All the MotionWarping windows within the supplied animation */
 	UFUNCTION(BlueprintCallable, Category = "Motion Warping")
 	static void GetMotionWarpingWindowsFromAnimation(const UAnimSequenceBase* Animation, TArray<FMotionWarpingWindowData>& OutWindows);
@@ -116,11 +119,14 @@ public:
 	bool ContainsModifier(const UAnimSequenceBase* Animation, float StartTime, float EndTime) const;
 
 	/** Add a new modifier */
-	void AddRootMotionModifier(TSharedPtr<FRootMotionModifier> Modifier);
+	FRootMotionModifierHandle AddRootMotionModifier(TSharedPtr<FRootMotionModifier> Modifier);
 
 	/** Mark all the modifiers as Disable */
 	void DisableAllRootMotionModifiers();
 
+	/** Finds a RootMotionModifier that matches the supplied Handle */
+	TSharedPtr<FRootMotionModifier> GetRootMotionModifierByHandle(const FRootMotionModifierHandle& Handle) const;
+
 protected:
 
 	/** Character this component belongs to */
@@ -143,4 +149,12 @@ protected:
 	FTransform ProcessRootMotionPreConvertToWorld(const FTransform& InRootMotion, class UCharacterMovementComponent* CharacterMovementComponent, float DeltaSeconds);
 	
 	FTransform ProcessRootMotionPostConvertToWorld(const FTransform& InRootMotion, class UCharacterMovementComponent* CharacterMovementComponent, float DeltaSeconds);
+
+public:
+
+	UFUNCTION(BlueprintCallable, Category = "Motion Warping")
+	bool SetRootMotionModifierDelegatesByHandle(FRootMotionModifierHandle Handle, FOnRootMotionModifierDelegate OnActivate, FOnRootMotionModifierDelegate OnUpdate, FOnRootMotionModifierDelegate OnDeactivate);
+
+	UFUNCTION(BlueprintCallable, Category = "Motion Warping")
+	bool GetAnimationAndTimeRangeForRootMotionModifierByHandle(FRootMotionModifierHandle Handle, UAnimSequenceBase*& OutAnimation, float& OutStartTime, float& OutEndTime) const;
 };
\ No newline at end of file
diff --git a/Engine/Plugins/Experimental/Animation/MotionWarping/Source/MotionWarping/Public/RootMotionModifier.h b/Engine/Plugins/Experimental/Animation/MotionWarping/Source/MotionWarping/Public/RootMotionModifier.h
--- a/Engine/Plugins/Experimental/Animation/MotionWarping/Source/MotionWarping/Public/RootMotionModifier.h
+++ b/Engine/Plugins/Experimental/Animation/MotionWarping/Source/MotionWarping/Public/RootMotionModifier.h
@@ -26,6 +26,36 @@ enum class ERootMotionModifierState : uint8
 	Disabled
 };
 
+/** Handle to identify a RootMotionModifier */
+USTRUCT(BlueprintType)
+struct FRootMotionModifierHandle
+{
+	GENERATED_BODY()
+
+	FRootMotionModifierHandle() : Handle(INDEX_NONE) {}
+
+	void GenerateNewHandle();
+	
+	bool IsValid() const { return Handle != INDEX_NONE; }
+	void Invalidate() { Handle = INDEX_NONE; }
+	
+	bool operator==(const FRootMotionModifierHandle& Other) const { return Handle == Other.Handle; }
+	bool operator!=(const FRootMotionModifierHandle& Other) const { return Handle != Other.Handle; }
+	
+	friend uint32 GetTypeHash(const FRootMotionModifierHandle& InHandle) { return InHandle.Handle; }
+
+	FString ToString() const { return IsValid() ? FString::FromInt(Handle) : TEXT("Invalid"); }
+
+	static const FRootMotionModifierHandle InvalidHandle;
+
+private:
+
+	UPROPERTY()
+	int32 Handle;
+};
+
+DECLARE_DYNAMIC_DELEGATE_TwoParams(FOnRootMotionModifierDelegate, UMotionWarpingComponent*, MotionWarpingComp, const FRootMotionModifierHandle&, Handle);
+
 /** Base struct for Root Motion Modifiers */
 USTRUCT()
 struct MOTIONWARPING_API FRootMotionModifier
@@ -38,10 +68,6 @@ public:
 	UPROPERTY()
 	TWeakObjectPtr<const UAnimSequenceBase> Animation = nullptr;
 
-	/** AnimNotifyState this modifier was created from. May be NULL */
-	UPROPERTY()
-	TWeakObjectPtr<const UAnimNotifyState_MotionWarping> AnimNotifyState = nullptr;
-
 	/** Start time of the warping window */
 	UPROPERTY()
 	float StartTime = 0.f;
@@ -70,26 +96,63 @@ public:
 	UPROPERTY()
 	ERootMotionModifierState State = ERootMotionModifierState::Waiting; //@TODO: Move to private
 
+	/** Delegate called when this modifier is activated (starts affecting the root motion) */
+	UPROPERTY()
+	FOnRootMotionModifierDelegate OnActivateDelegate;
+
+	/** Delegate called when this modifier updates while active (affecting the root motion) */
+	UPROPERTY()
+	FOnRootMotionModifierDelegate OnUpdateDelegate;
+
+	/** Delegate called when this modifier is deactivated (stops affecting the root motion) */
+	UPROPERTY()
+	FOnRootMotionModifierDelegate OnDeactivateDelegate;
+
 	FRootMotionModifier(){}
-	virtual ~FRootMotionModifier() {}
+	virtual ~FRootMotionModifier(){}
+
+	/** Initialize this modifier */
+	void Initialize(UMotionWarpingComponent* OwnerComp);
 
 	/** Returns the actual struct used for RTTI */
 	virtual UScriptStruct* GetScriptStruct() const PURE_VIRTUAL(FRootMotionModifier::GetScriptStruct, return FRootMotionModifier::StaticStruct(););
 
-	/** Updates the state of the modifier. Runs before ProcessRootMotion */
-	virtual void Update(UMotionWarpingComponent& OwnerComp);
-
-	/** Performs the actual modification to the motion */
-	virtual FTransform ProcessRootMotion(UMotionWarpingComponent& OwnerComp, const FTransform& InRootMotion, float DeltaSeconds) PURE_VIRTUAL(FRootMotionModifier::ProcessRootMotion, return FTransform::Identity;);
-
 	/** Called when the state of the modifier changes */
-	virtual void OnStateChanged(ERootMotionModifierState LastState){}
+	virtual void OnStateChanged(ERootMotionModifierState LastState);
 
 	/** Sets the state of the modifier */
 	void SetState(ERootMotionModifierState NewState);
 
 	/** Returns the state of the modifier */
 	FORCEINLINE ERootMotionModifierState GetState() const { return State; }
+
+	/** Returns the handle to identify this modifier */
+	FORCEINLINE const FRootMotionModifierHandle& GetHandle() const { return Handle; }
+
+	/** Returns a pointer to the component that owns this modifier */
+	FORCEINLINE UMotionWarpingComponent* GetOwnerComponent() const { return OwnerComponent.Get(); }
+
+	/** Returns a pointer to the character that owns the component that owns this modifier */
+	class ACharacter* GetCharacterOwner() const;
+
+	//DEPRECATED will be removed shortly, replaced with a signature that does not receive the comp via param.
+	virtual void Update(UMotionWarpingComponent& OwnerComp);
+	virtual FTransform ProcessRootMotion(UMotionWarpingComponent& OwnerComp, const FTransform& InRootMotion, float DeltaSeconds) { return FTransform::Identity;	}
+
+	FORCEINLINE const UAnimSequenceBase* GetAnimation() const { return Animation.Get(); }
+
+private:
+
+	friend UMotionWarpingComponent;
+
+	/** Pointer back to the Component that owns this modifier */
+	UPROPERTY()
+	TWeakObjectPtr<UMotionWarpingComponent> OwnerComponent = nullptr;
+
+	/** Handle to identify this modifier */
+	UPROPERTY()
+	FRootMotionModifierHandle Handle;
+
 };
 
 /** Blueprint wrapper around the config properties of a root motion modifier */
@@ -104,7 +167,10 @@ public:
 		: Super(ObjectInitializer) {}
 
 	/** Adds a RootMotionModifier of the type this object represents */
-	virtual void AddRootMotionModifier(UMotionWarpingComponent* InMotionWarpingComp, const UAnimSequenceBase* Animation, float StartTime, float EndTime) const PURE_VIRTUAL(URootMotionModifierConfig::AddRootMotionModifier,);
+	virtual FRootMotionModifierHandle AddRootMotionModifierNew(UMotionWarpingComponent* InMotionWarpingComp, const UAnimSequenceBase* Animation, float StartTime, float EndTime) const { return FRootMotionModifierHandle::InvalidHandle; }
+
+	//DEPRECATED will be removed shortly
+	virtual void AddRootMotionModifier(UMotionWarpingComponent* InMotionWarpingComp, const UAnimSequenceBase* Animation, float StartTime, float EndTime) const {}
 };
 
 /** Represents a point of alignment in the world */
@@ -158,6 +224,20 @@ enum class EMotionWarpRotationType : uint8
 	Facing,
 };
 
+/** Method used to extract the warp point from the animation */
+UENUM(BlueprintType)
+enum class EWarpPointAnimProvider : uint8
+{
+	/** No warp point is provided */
+	None,
+
+	/** Warp point defined by a 'hard-coded' transform  user can enter through the warping notify */
+	Static,
+
+	/** Warp point defined by a bone */
+	Bone
+};
+
 USTRUCT()
 struct MOTIONWARPING_API FRootMotionModifier_Warp : public FRootMotionModifier
 {
@@ -165,9 +245,19 @@ struct MOTIONWARPING_API FRootMotionModifier_Warp : public FRootMotionModifier
 
 public:
 
-	/** Name used to find the sync point for this modifier */
+	/** Name used to find the warp point for this modifier */
+	UPROPERTY()
+	FName SyncPointName = NAME_None; //@TODO: Rename to WarpPointName
+
+	/** Method used to extract the warp point from the animation */
+	UPROPERTY()
+	EWarpPointAnimProvider WarpPointAnimProvider = EWarpPointAnimProvider::None;
+
 	UPROPERTY()
-	FName SyncPointName = NAME_None;
+	FTransform WarpPointAnimTransform = FTransform::Identity;
+
+	UPROPERTY()
+	FName WarpPointAnimBoneName = NAME_None;
 
 	/** Whether to warp the translation component of the root motion */
 	UPROPERTY()
@@ -185,14 +275,14 @@ public:
 	UPROPERTY()
 	EMotionWarpRotationType RotationType = EMotionWarpRotationType::Default;
 
-	/** 
-	 * Allow to modify how fast the rotation is warped. 
+	/**
+	 * Allow to modify how fast the rotation is warped.
 	 * e.g if the window duration is 2sec and this is 0.5, the target rotation will be reached in 1sec instead of 2sec
 	 */
 	UPROPERTY()
 	float WarpRotationTimeMultiplier = 1.f;
 
-	/** Sync Point used by this modifier as target for the warp. Cached during the Update */
+	//DEPRECATED: Will be removed shortly
 	UPROPERTY()
 	FMotionWarpingSyncPoint CachedSyncPoint;
 
@@ -205,17 +295,36 @@ public:
 	virtual FTransform ProcessRootMotion(UMotionWarpingComponent& OwnerComp, const FTransform& InRootMotion, float DeltaSeconds) override;
 	//~ End FRootMotionModifier Interface
 
-	/** Event called during update if the sync point changes while the warping is active */
-	virtual void OnSyncPointChanged(UMotionWarpingComponent& OwnerComp) {}
+	/** Event called during update if the target transform changes while the warping is active */
+	virtual void OnTargetTransformChanged() {}
 
 #if !(UE_BUILD_SHIPPING || UE_BUILD_TEST)
-	void PrintLog(const UMotionWarpingComponent& OwnerComp, const FString& Name, const FTransform& OriginalRootMotion, const FTransform& WarpedRootMotion) const;
+	void PrintLog(const FString& Name, const FTransform& OriginalRootMotion, const FTransform& WarpedRootMotion) const;
+
+	//DEPRECATED will be removed shortly
+	void PrintLog(UMotionWarpingComponent& OwnerComp, const FString& Name, const FTransform& OriginalRootMotion, const FTransform& WarpedRootMotion) const
+	{
+		PrintLog(Name, OriginalRootMotion, WarpedRootMotion);
+	}
+
 #endif
 
 protected:
 
-	FQuat GetTargetRotation(UMotionWarpingComponent& OwnerComp) const;
-	FQuat WarpRotation(UMotionWarpingComponent& OwnerComp, const FTransform& RootMotionDelta, const FTransform& RootMotionTotal, float DeltaSeconds);
+	FORCEINLINE FVector GetTargetLocation() const { return CachedTargetTransform.GetLocation(); }
+	FORCEINLINE FRotator GetTargetRotator() const { return GetTargetRotation().Rotator(); }
+	FQuat GetTargetRotation() const;
+	
+	FQuat WarpRotation(const FTransform& RootMotionDelta, const FTransform& RootMotionTotal, float DeltaSeconds);
+
+	//@TODO: This should be Optional and private
+	UPROPERTY()
+	FTransform CachedTargetTransform = FTransform::Identity;
+
+private:
+
+	/** Cached of the offset from the warp point. Used to calculate the final target transform when a warp point is defined in the animation */
+	TOptional<FTransform> CachedOffsetFromWarpPoint;
 };
 
 UCLASS(meta = (DisplayName = "Simple Warp"))
@@ -225,9 +334,20 @@ class MOTIONWARPING_API URootMotionModifierConfig_Warp : public URootMotionModif
 
 public:
 
-	/** Name used to find the sync point for this modifier */
+	/** Name used to find the warp point for this modifier */
+	UPROPERTY(EditAnywhere, BlueprintReadWrite, Category = "Config", meta = (DisplayName = "Warp Point Name"))
+	FName SyncPointName = NAME_None;  //@TODO: Rename to WarpPointName
+
 	UPROPERTY(EditAnywhere, BlueprintReadWrite, Category = "Config")
-	FName SyncPointName = NAME_None;
+	EWarpPointAnimProvider WarpPointAnimProvider = EWarpPointAnimProvider::None;
+
+	//@TODO: Hide from the UI when Provider != Static
+	UPROPERTY(EditAnywhere, BlueprintReadWrite, Category = "Config", meta = (EditCondition = "WarpPointAnimProvider == EWarpPointAnimProvider::Static"))
+	FTransform WarpPointAnimTransform = FTransform::Identity;
+
+	//@TODO: Hide from the UI when Provider != Bone
+	UPROPERTY(EditAnywhere, BlueprintReadWrite, Category = "Config", meta = (EditCondition = "WarpPointAnimProvider == EWarpPointAnimProvider::Bone"))
+	FName WarpPointAnimBoneName = NAME_None;
 
 	/** Whether to warp the translation component of the root motion */
 	UPROPERTY(EditAnywhere, BlueprintReadWrite, Category = "Config")
@@ -254,13 +374,17 @@ public:
 
 	URootMotionModifierConfig_Warp(const FObjectInitializer& ObjectInitializer) : Super(ObjectInitializer) {}
 
-	virtual void AddRootMotionModifier(UMotionWarpingComponent* MotionWarpingComp, const UAnimSequenceBase* Animation, float StartTime, float EndTime) const override
+	virtual FRootMotionModifierHandle AddRootMotionModifierNew(UMotionWarpingComponent* MotionWarpingComp, const UAnimSequenceBase* Animation, float StartTime, float EndTime) const override
 	{
-		URootMotionModifierConfig_Warp::AddRootMotionModifierSimpleWarp(MotionWarpingComp, Animation, StartTime, EndTime, SyncPointName, bWarpTranslation, bIgnoreZAxis, bWarpRotation, RotationType, WarpRotationTimeMultiplier);
+		return URootMotionModifierConfig_Warp::AddRootMotionModifierSimpleWarp(MotionWarpingComp, Animation, StartTime, EndTime, 
+			SyncPointName, WarpPointAnimProvider, WarpPointAnimTransform, WarpPointAnimBoneName,
+			bWarpTranslation, bIgnoreZAxis, bWarpRotation, RotationType, WarpRotationTimeMultiplier);
 	}
 
 	UFUNCTION(BlueprintCallable, Category = "Motion Warping")
-	static void AddRootMotionModifierSimpleWarp(UMotionWarpingComponent* InMotionWarpingComp, const UAnimSequenceBase* InAnimation, float InStartTime, float InEndTime, FName InSyncPointName, bool bInWarpTranslation, bool bInIgnoreZAxis, bool bInWarpRotation, EMotionWarpRotationType InRotationType, float InWarpRotationTimeMultiplier = 1.f);
+	static FRootMotionModifierHandle AddRootMotionModifierSimpleWarp(UMotionWarpingComponent* InMotionWarpingComp, const UAnimSequenceBase* InAnimation, float InStartTime, float InEndTime,
+		FName InSyncPointName, EWarpPointAnimProvider InWarpPointAnimProvider, FTransform InWarpPointAnimTransform, FName InWarpPointAnimBoneName,
+		bool bInWarpTranslation, bool bInIgnoreZAxis, bool bInWarpRotation, EMotionWarpRotationType InRotationType, float InWarpRotationTimeMultiplier = 1.f);
 };
 
 // FRootMotionModifier_Scale
@@ -302,11 +426,11 @@ public:
 
 	URootMotionModifierConfig_Scale(const FObjectInitializer& ObjectInitializer) : Super(ObjectInitializer) {}
 
-	virtual void AddRootMotionModifier(UMotionWarpingComponent* MotionWarpingComp, const UAnimSequenceBase* Animation, float StartTime, float EndTime) const override
+	virtual FRootMotionModifierHandle AddRootMotionModifierNew(UMotionWarpingComponent* MotionWarpingComp, const UAnimSequenceBase* Animation, float StartTime, float EndTime) const override
 	{
-		URootMotionModifierConfig_Scale::AddRootMotionModifierScale(MotionWarpingComp, Animation, StartTime, EndTime, Scale);
+		return URootMotionModifierConfig_Scale::AddRootMotionModifierScale(MotionWarpingComp, Animation, StartTime, EndTime, Scale);
 	}
 
 	UFUNCTION(BlueprintCallable, Category = "Motion Warping")
-	static void AddRootMotionModifierScale(UMotionWarpingComponent* InMotionWarpingComp, const UAnimSequenceBase* InAnimation, float InStartTime, float InEndTime, FVector InScale);
+	static FRootMotionModifierHandle AddRootMotionModifierScale(UMotionWarpingComponent* InMotionWarpingComp, const UAnimSequenceBase* InAnimation, float InStartTime, float InEndTime, FVector InScale);
 };
\ No newline at end of file
diff --git a/Engine/Plugins/Experimental/Animation/MotionWarping/Source/MotionWarping/Public/RootMotionModifier_AdjustmentBlendWarp.h b/Engine/Plugins/Experimental/Animation/MotionWarping/Source/MotionWarping/Public/RootMotionModifier_AdjustmentBlendWarp.h
--- a/Engine/Plugins/Experimental/Animation/MotionWarping/Source/MotionWarping/Public/RootMotionModifier_AdjustmentBlendWarp.h
+++ b/Engine/Plugins/Experimental/Animation/MotionWarping/Source/MotionWarping/Public/RootMotionModifier_AdjustmentBlendWarp.h
@@ -58,7 +58,7 @@ public:
 	virtual ~FRootMotionModifier_AdjustmentBlendWarp() {}
 
 	virtual UScriptStruct* GetScriptStruct() const { return FRootMotionModifier_AdjustmentBlendWarp::StaticStruct(); }
-	virtual void OnSyncPointChanged(UMotionWarpingComponent& OwnerComp) override;
+	virtual void OnTargetTransformChanged() override;
 	virtual FTransform ProcessRootMotion(UMotionWarpingComponent& OwnerComp, const FTransform& InRootMotion, float DeltaSeconds) override;
 
 	void GetIKBoneTransformAndAlpha(FName BoneName, FTransform& OutTransform, float& OutAlpha) const;
@@ -80,7 +80,7 @@ protected:
 	UPROPERTY()
 	float ActualStartTime = 0.f;
 
-	void PrecomputeWarpedTracks(UMotionWarpingComponent& OwnerComp);
+	void PrecomputeWarpedTracks();
 
 	FTransform ExtractWarpedRootMotion() const;
 
@@ -89,7 +89,7 @@ protected:
 	void ExtractBoneTransformAtFrame(FTransform& OutTransform, int32 TrackIndex, int32 Frame) const;
 
 #if !(UE_BUILD_SHIPPING || UE_BUILD_TEST)
-	void DrawDebugWarpedTracks(UMotionWarpingComponent& OwnerComp, float DrawDuration) const;
+	void DrawDebugWarpedTracks(float DrawDuration) const;
 #endif
 
 	static void ExtractMotionDeltaFromRange(const FBoneContainer& BoneContainer, const UAnimSequenceBase* Animation, float StartTime, float EndTime, float SampleRate, FMotionDeltaTrackContainer& OutMotionDeltaTracks);
@@ -113,13 +113,17 @@ public:
 	URootMotionModifierConfig_AdjustmentBlendWarp(const FObjectInitializer& ObjectInitializer)
 		: Super(ObjectInitializer) {}
 
-	virtual void AddRootMotionModifier(UMotionWarpingComponent* MotionWarpingComp, const UAnimSequenceBase* Animation, float StartTime, float EndTime) const override
+	virtual FRootMotionModifierHandle AddRootMotionModifierNew(UMotionWarpingComponent* MotionWarpingComp, const UAnimSequenceBase* Animation, float StartTime, float EndTime) const override
 	{
-		URootMotionModifierConfig_AdjustmentBlendWarp::AddRootMotionModifierAdjustmentBlendWarp(MotionWarpingComp, Animation, StartTime, EndTime, SyncPointName, bWarpTranslation, bIgnoreZAxis, bWarpRotation, bWarpIKBones, IKBones);
+		return URootMotionModifierConfig_AdjustmentBlendWarp::AddRootMotionModifierAdjustmentBlendWarp(MotionWarpingComp, Animation, StartTime, EndTime, 
+			SyncPointName, WarpPointAnimProvider, WarpPointAnimTransform, WarpPointAnimBoneName,
+			bWarpTranslation, bIgnoreZAxis, bWarpRotation, bWarpIKBones, IKBones);
 	}
 
 	UFUNCTION(BlueprintCallable, Category = "Motion Warping")
-	static void AddRootMotionModifierAdjustmentBlendWarp(UMotionWarpingComponent* InMotionWarpingComp, const UAnimSequenceBase* InAnimation, float InStartTime, float InEndTime, FName InSyncPointName, bool bInWarpTranslation, bool bInIgnoreZAxis, bool bInWarpRotation, bool bInWarpIKBones, const TArray<FName>& InIKBones);
+	static FRootMotionModifierHandle AddRootMotionModifierAdjustmentBlendWarp(UMotionWarpingComponent* InMotionWarpingComp, const UAnimSequenceBase* InAnimation, float InStartTime, float InEndTime, 
+		FName InSyncPointName, EWarpPointAnimProvider InWarpPointAnimProvider, FTransform InWarpPointAnimTransform, FName InWarpPointAnimBoneName,
+		bool bInWarpTranslation, bool bInIgnoreZAxis, bool bInWarpRotation, bool bInWarpIKBones, const TArray<FName>& InIKBones);
 
 	UFUNCTION(BlueprintPure, Category = "Motion Warping")
 	static void GetIKBoneTransformAndAlpha(ACharacter* Character, FName BoneName, FTransform& OutTransform, float& OutAlpha);
diff --git a/Engine/Plugins/Experimental/Animation/MotionWarping/Source/MotionWarping/Public/RootMotionModifier_SkewWarp.h b/Engine/Plugins/Experimental/Animation/MotionWarping/Source/MotionWarping/Public/RootMotionModifier_SkewWarp.h
--- a/Engine/Plugins/Experimental/Animation/MotionWarping/Source/MotionWarping/Public/RootMotionModifier_SkewWarp.h
+++ b/Engine/Plugins/Experimental/Animation/MotionWarping/Source/MotionWarping/Public/RootMotionModifier_SkewWarp.h
@@ -30,11 +30,15 @@ public:
 	URootMotionModifierConfig_SkewWarp(const FObjectInitializer& ObjectInitializer)
 		: Super(ObjectInitializer) {}
 
-	virtual void AddRootMotionModifier(UMotionWarpingComponent* MotionWarpingComp, const UAnimSequenceBase* Animation, float StartTime, float EndTime) const override
+	virtual FRootMotionModifierHandle AddRootMotionModifierNew(UMotionWarpingComponent* MotionWarpingComp, const UAnimSequenceBase* Animation, float StartTime, float EndTime) const override
 	{
-		URootMotionModifierConfig_SkewWarp::AddRootMotionModifierSkewWarp(MotionWarpingComp, Animation, StartTime, EndTime, SyncPointName, bWarpTranslation, bIgnoreZAxis, bWarpRotation, RotationType, WarpRotationTimeMultiplier);
+		return URootMotionModifierConfig_SkewWarp::AddRootMotionModifierSkewWarp(MotionWarpingComp, Animation, StartTime, EndTime, 
+			SyncPointName, WarpPointAnimProvider, WarpPointAnimTransform, WarpPointAnimBoneName,
+			bWarpTranslation, bIgnoreZAxis, bWarpRotation, RotationType, WarpRotationTimeMultiplier);
 	}
 
 	UFUNCTION(BlueprintCallable, Category = "Motion Warping")
-	static void AddRootMotionModifierSkewWarp(UMotionWarpingComponent* InMotionWarpingComp, const UAnimSequenceBase* InAnimation, float InStartTime, float InEndTime, FName InSyncPointName, bool bInWarpTranslation, bool bInIgnoreZAxis, bool bInWarpRotation, EMotionWarpRotationType InRotationType, float InWarpRotationTimeMultiplier = 1.f);
+	static FRootMotionModifierHandle AddRootMotionModifierSkewWarp(UMotionWarpingComponent* InMotionWarpingComp, const UAnimSequenceBase* InAnimation, float InStartTime, float InEndTime, 
+		FName InSyncPointName, EWarpPointAnimProvider InWarpPointAnimProvider, FTransform InWarpPointAnimTransform, FName InWarpPointAnimBoneName,
+		bool bInWarpTranslation, bool bInIgnoreZAxis, bool bInWarpRotation, EMotionWarpRotationType InRotationType, float InWarpRotationTimeMultiplier = 1.f);
 };
\ No newline at end of file
