Motion Warping - Exposed FRootMotionModifierHandle #jira none #review trivial

diff --git a/Engine/Plugins/Experimental/Animation/MotionWarping/Source/MotionWarping/Public/RootMotionModifier.h b/Engine/Plugins/Experimental/Animation/MotionWarping/Source/MotionWarping/Public/RootMotionModifier.h
--- a/Engine/Plugins/Experimental/Animation/MotionWarping/Source/MotionWarping/Public/RootMotionModifier.h
+++ b/Engine/Plugins/Experimental/Animation/MotionWarping/Source/MotionWarping/Public/RootMotionModifier.h
@@ -28,7 +28,7 @@ enum class ERootMotionModifierState : uint8
 
 /** Handle to identify a RootMotionModifier */
 USTRUCT(BlueprintType)
-struct FRootMotionModifierHandle
+struct MOTIONWARPING_API FRootMotionModifierHandle
 {
 	GENERATED_BODY()
 
