Motion Warping - Fix for static analysis warning #jira UE-111681 #rb trivial #fyi [REDACTED_HANDLE]

diff --git a/Engine/Plugins/Experimental/Animation/MotionWarping/Source/MotionWarping/Private/RootMotionModifier.cpp b/Engine/Plugins/Experimental/Animation/MotionWarping/Source/MotionWarping/Private/RootMotionModifier.cpp
--- a/Engine/Plugins/Experimental/Animation/MotionWarping/Source/MotionWarping/Private/RootMotionModifier.cpp
+++ b/Engine/Plugins/Experimental/Animation/MotionWarping/Source/MotionWarping/Private/RootMotionModifier.cpp
@@ -126,10 +126,10 @@ void FRootMotionModifier_Warp::Update(UMotionWarpingComponent& OwnerComp_DEPRECA
 	FRootMotionModifier::Update(OwnerComp_DEPRECATED);
 
 	// Cache sync point transform and trigger OnSyncPointChanged if needed
-	if (GetState() == ERootMotionModifierState::Active)
+	const UMotionWarpingComponent* OwnerComp = GetOwnerComponent();
+	if (OwnerComp && GetState() == ERootMotionModifierState::Active)
 	{
-		UMotionWarpingComponent* OwnerComp = GetOwnerComponent();
-		const FMotionWarpingSyncPoint* SyncPointPtr = OwnerComp ? OwnerComp->FindSyncPoint(SyncPointName) : nullptr;
+		const FMotionWarpingSyncPoint* SyncPointPtr = OwnerComp->FindSyncPoint(SyncPointName);
 		
 		// Disable if there is no sync point for us
 		if (SyncPointPtr == nullptr)
