Motion Warping - Added UObject version of root motion modifiers (temporally named MotionModifier to avoid name collision while FRootMotionModifier is still alive for backward compatibility)

diff --git a/Engine/Plugins/Experimental/Animation/MotionWarping/Source/MotionWarping/Private/AnimNotifyState_MotionWarping.cpp b/Engine/Plugins/Experimental/Animation/MotionWarping/Source/MotionWarping/Private/AnimNotifyState_MotionWarping.cpp
--- a/Engine/Plugins/Experimental/Animation/MotionWarping/Source/MotionWarping/Private/AnimNotifyState_MotionWarping.cpp
+++ b/Engine/Plugins/Experimental/Animation/MotionWarping/Source/MotionWarping/Private/AnimNotifyState_MotionWarping.cpp
@@ -14,7 +14,7 @@ FRootMotionModifierHandle UAnimNotifyState_MotionWarping::OnBecomeRelevant(UMoti
 	FRootMotionModifierHandle Handle = AddRootMotionModifier(MotionWarpingComp, Animation, StartTime, EndTime);
 
 	//@TODO: Temp, to handle the case where the blueprint notify does not return the handle. Will be removed shortly
-	if(!Handle.IsValid())
+	if (!Handle.IsValid())
 	{
 		if (MotionWarpingComp->GetRootMotionModifiers().Num())
 		{
@@ -24,12 +24,24 @@ FRootMotionModifierHandle UAnimNotifyState_MotionWarping::OnBecomeRelevant(UMoti
 				Handle = Last->GetHandle();
 			}
 		}
+
+		if (!Handle.IsValid())
+		{
+			if (MotionWarpingComp->GetModifiers().Num())
+			{
+				UMotionModifier* Last = MotionWarpingComp->GetModifiers().Last();
+				if (Last && Last->Animation.Get() == Animation && Last->StartTime == StartTime && Last->EndTime == EndTime)
+				{
+					Handle = Last->GetHandle();
+				}
+			}
+		}
 	}
 
 	TSharedPtr<FRootMotionModifier> Modifier = MotionWarpingComp->GetRootMotionModifierByHandle(Handle);
 	if (Modifier.IsValid())
 	{
-		if(!Modifier->OnActivateDelegate.IsBound())
+		if (!Modifier->OnActivateDelegate.IsBound())
 		{
 			Modifier->OnActivateDelegate.BindDynamic(this, &UAnimNotifyState_MotionWarping::OnRootMotionModifierActivate);
 		}
@@ -43,6 +55,28 @@ FRootMotionModifierHandle UAnimNotifyState_MotionWarping::OnBecomeRelevant(UMoti
 		{
 			Modifier->OnDeactivateDelegate.BindDynamic(this, &UAnimNotifyState_MotionWarping::OnRootMotionModifierDeactivate);
 		}
+
+		return Handle;
+	}
+
+	if (UMotionModifier* RootMotionModifierNew = MotionWarpingComp->GetModifierByHandle(Handle))
+	{
+		if (!RootMotionModifierNew->OnActivateDelegate.IsBound())
+		{
+			RootMotionModifierNew->OnActivateDelegate.BindDynamic(this, &UAnimNotifyState_MotionWarping::OnRootMotionModifierActivate);
+		}
+
+		if (!RootMotionModifierNew->OnUpdateDelegate.IsBound())
+		{
+			RootMotionModifierNew->OnUpdateDelegate.BindDynamic(this, &UAnimNotifyState_MotionWarping::OnRootMotionModifierUpdate);
+		}
+
+		if (!RootMotionModifierNew->OnDeactivateDelegate.IsBound())
+		{
+			RootMotionModifierNew->OnDeactivateDelegate.BindDynamic(this, &UAnimNotifyState_MotionWarping::OnRootMotionModifierDeactivate);
+		}
+
+		return Handle;
 	}
 
 	return Handle;
@@ -50,9 +84,20 @@ FRootMotionModifierHandle UAnimNotifyState_MotionWarping::OnBecomeRelevant(UMoti
 
 FRootMotionModifierHandle UAnimNotifyState_MotionWarping::AddRootMotionModifier_Implementation(UMotionWarpingComponent* MotionWarpingComp, const UAnimSequenceBase* Animation, float StartTime, float EndTime) const
 {
-	if (MotionWarpingComp && RootMotionModifierConfig)
+	if (MotionWarpingComp)
 	{
-		return RootMotionModifierConfig->AddRootMotionModifierNew(MotionWarpingComp, Animation, StartTime, EndTime);
+		if (RootMotionModifierConfig)
+		{
+			return RootMotionModifierConfig->AddRootMotionModifierNew(MotionWarpingComp, Animation, StartTime, EndTime);
+		}
+		else if (RootMotionModifier)
+		{
+			UMotionModifier* NewRootMotionModifier = MotionWarpingComp->AddModifierFromTemplate(RootMotionModifier, Animation, StartTime, EndTime);
+			if (NewRootMotionModifier)
+			{
+				return NewRootMotionModifier->GetHandle();
+			}
+		}
 	}
 
 	return FRootMotionModifierHandle::InvalidHandle;
diff --git a/Engine/Plugins/Experimental/Animation/MotionWarping/Source/MotionWarping/Private/MotionWarpingComponent.cpp b/Engine/Plugins/Experimental/Animation/MotionWarping/Source/MotionWarping/Private/MotionWarpingComponent.cpp
--- a/Engine/Plugins/Experimental/Animation/MotionWarping/Source/MotionWarping/Private/MotionWarpingComponent.cpp
+++ b/Engine/Plugins/Experimental/Animation/MotionWarping/Source/MotionWarping/Private/MotionWarpingComponent.cpp
@@ -172,9 +172,18 @@ void UMotionWarpingComponent::InitializeComponent()
 
 bool UMotionWarpingComponent::ContainsModifier(const UAnimSequenceBase* Animation, float StartTime, float EndTime) const
 {
-	return RootMotionModifiers.ContainsByPredicate([=](const TSharedPtr<FRootMotionModifier>& Modifier)
+	// DEPRECATED will be removed shortly
+	if (RootMotionModifiers.ContainsByPredicate([=](const TSharedPtr<FRootMotionModifier>& Modifier)
 		{
 			return (Modifier.IsValid() && Modifier->Animation == Animation && Modifier->StartTime == StartTime && Modifier->EndTime == EndTime);
+		}))
+	{
+		return true;
+	}
+
+	return Modifiers.ContainsByPredicate([=](const UMotionModifier* Modifier)
+		{
+			return (Modifier->Animation == Animation && Modifier->StartTime == StartTime && Modifier->EndTime == EndTime);
 		});
 }
 
@@ -197,6 +206,25 @@ FRootMotionModifierHandle UMotionWarpingComponent::AddRootMotionModifier(TShared
 	return FRootMotionModifierHandle::InvalidHandle;
 }
 
+FRootMotionModifierHandle UMotionWarpingComponent::AddModifier(UMotionModifier* Modifier)
+{
+	if (ensureAlways(Modifier))
+	{
+		//@TODO: Move this from here after update all modifiers
+		Modifier->Initialize(this);
+
+		Modifiers.Add(Modifier);
+
+		UE_LOG(LogMotionWarping, Verbose, TEXT("MotionWarping: RootMotionModifier added. NetMode: %d WorldTime: %f Char: %s Animation: %s [%f %f] [%f %f] Loc: %s Rot: %s"),
+			GetWorld()->GetNetMode(), GetWorld()->GetTimeSeconds(), *GetNameSafe(GetCharacterOwner()), *GetNameSafe(Modifier->Animation.Get()), Modifier->StartTime, Modifier->EndTime, Modifier->PreviousPosition, Modifier->CurrentPosition,
+			*GetCharacterOwner()->GetActorLocation().ToString(), *GetCharacterOwner()->GetActorRotation().ToCompactString());
+
+		return Modifier->GetHandle();
+	}
+
+	return FRootMotionModifierHandle::InvalidHandle;
+}
+
 TSharedPtr<FRootMotionModifier> UMotionWarpingComponent::GetRootMotionModifierByHandle(const FRootMotionModifierHandle& Handle) const
 {
 	if(Handle.IsValid())
@@ -213,8 +241,25 @@ TSharedPtr<FRootMotionModifier> UMotionWarpingComponent::GetRootMotionModifierBy
 	return TSharedPtr<FRootMotionModifier>(nullptr);
 }
 
+UMotionModifier* UMotionWarpingComponent::GetModifierByHandle(const FRootMotionModifierHandle& Handle) const
+{
+	if (Handle.IsValid())
+	{
+		for (UMotionModifier* RootMotionModifier : Modifiers)
+		{
+			if (RootMotionModifier && RootMotionModifier->GetHandle() == Handle)
+			{
+				return RootMotionModifier;
+			}
+		}
+	}
+
+	return nullptr;
+}
+
 void UMotionWarpingComponent::DisableAllRootMotionModifiers()
 {
+	// DEPRECATED will be removed shortly
 	if (RootMotionModifiers.Num() > 0)
 	{
 		for (TSharedPtr<FRootMotionModifier>& Modifier : RootMotionModifiers)
@@ -222,6 +267,14 @@ void UMotionWarpingComponent::DisableAllRootMotionModifiers()
 			Modifier->SetState(ERootMotionModifierState::Disabled);
 		}
 	}
+
+	if (Modifiers.Num() > 0)
+	{
+		for (UMotionModifier* Modifier : Modifiers)
+		{
+			Modifier->SetState(ERootMotionModifierState::Disabled);
+		}
+	}
 }
 
 void UMotionWarpingComponent::Update()
@@ -315,6 +368,30 @@ void UMotionWarpingComponent::Update()
 			return false; 
 		});
 	}
+
+	// Update the state of all the modifiers
+	if (Modifiers.Num() > 0)
+	{
+		for (UMotionModifier* Modifier : Modifiers)
+		{
+			Modifier->Update();
+		}
+
+		// Remove the modifiers that has been marked for removal
+		Modifiers.RemoveAll([this](const UMotionModifier* Modifier)
+		{
+			if (Modifier->GetState() == ERootMotionModifierState::MarkedForRemoval)
+			{
+				UE_LOG(LogMotionWarping, Verbose, TEXT("MotionWarping: RootMotionModifier removed. NetMode: %d WorldTime: %f Char: %s Animation: %s [%f %f] [%f %f] Loc: %s Rot: %s"),
+					GetWorld()->GetNetMode(), GetWorld()->GetTimeSeconds(), *GetNameSafe(GetCharacterOwner()), *GetNameSafe(Modifier->Animation.Get()), Modifier->StartTime, Modifier->EndTime, Modifier->PreviousPosition, Modifier->CurrentPosition,
+					*GetCharacterOwner()->GetActorLocation().ToString(), *GetCharacterOwner()->GetActorRotation().ToCompactString());
+
+				return true;
+			}
+
+			return false;
+		});
+	}
 }
 
 FTransform UMotionWarpingComponent::ProcessRootMotionPreConvertToWorld(const FTransform& InRootMotion, UCharacterMovementComponent* CharacterMovementComponent, float DeltaSeconds)
@@ -341,6 +418,15 @@ FTransform UMotionWarpingComponent::ProcessRootMotionPreConvertToWorld(const FTr
 		}
 	}
 
+	// Apply Local Space Modifiers
+	for (UMotionModifier* Modifier : Modifiers)
+	{
+		if (Modifier->GetState() == ERootMotionModifierState::Active && Modifier->bInLocalSpace)
+		{
+			FinalRootMotion = Modifier->ProcessRootMotion(FinalRootMotion, DeltaSeconds);
+		}
+	}
+
 	return FinalRootMotion;
 }
 
@@ -365,6 +451,15 @@ FTransform UMotionWarpingComponent::ProcessRootMotionPostConvertToWorld(const FT
 		}
 	}
 
+	// Apply World Space Modifiers
+	for (UMotionModifier* Modifier : Modifiers)
+	{
+		if (Modifier->GetState() == ERootMotionModifierState::Active && !Modifier->bInLocalSpace)
+		{
+			FinalRootMotion = Modifier->ProcessRootMotion(FinalRootMotion, DeltaSeconds);
+		}
+	}
+
 #if !(UE_BUILD_SHIPPING || UE_BUILD_TEST)
 	const int32 DebugLevel = FMotionWarpingCVars::CVarMotionWarpingDebug.GetValueOnGameThread();
 	if (DebugLevel >= 2)
@@ -415,6 +510,7 @@ int32 UMotionWarpingComponent::RemoveSyncPoint(FName Name)
 
 bool UMotionWarpingComponent::SetRootMotionModifierDelegatesByHandle(FRootMotionModifierHandle Handle, FOnRootMotionModifierDelegate OnActivate, FOnRootMotionModifierDelegate OnUpdate, FOnRootMotionModifierDelegate OnDeactivate)
 {
+	// DEPRECATED will be removed shortly
 	TSharedPtr<FRootMotionModifier> Modifier = GetRootMotionModifierByHandle(Handle);
 	if (Modifier.IsValid())
 	{
@@ -424,11 +520,20 @@ bool UMotionWarpingComponent::SetRootMotionModifierDelegatesByHandle(FRootMotion
 		return true;
 	}
 
+	if (UMotionModifier* RootMotionModifier = GetModifierByHandle(Handle))
+	{
+		RootMotionModifier->OnActivateDelegate = OnActivate;
+		RootMotionModifier->OnUpdateDelegate = OnUpdate;
+		RootMotionModifier->OnDeactivateDelegate = OnDeactivate;
+		return true;
+	}
+
 	return false;
 }
 
 bool UMotionWarpingComponent::GetAnimationAndTimeRangeForRootMotionModifierByHandle(FRootMotionModifierHandle Handle, UAnimSequenceBase*& OutAnimation, float& OutStartTime, float& OutEndTime) const
 {
+	// DEPRECATED will be removed shortly
 	TSharedPtr<FRootMotionModifier> Modifier = GetRootMotionModifierByHandle(Handle);
 	if (Modifier.IsValid())
 	{
@@ -438,5 +543,32 @@ bool UMotionWarpingComponent::GetAnimationAndTimeRangeForRootMotionModifierByHan
 		return true;
 	}
 
+	if (UMotionModifier* RootMotionModifier = GetModifierByHandle(Handle))
+	{
+		OutAnimation = const_cast<UAnimSequenceBase*>(RootMotionModifier->GetAnimation());
+		OutStartTime = RootMotionModifier->StartTime;
+		OutEndTime = RootMotionModifier->EndTime;
+		return true;
+	}
+
 	return false;
+}
+
+UMotionModifier* UMotionWarpingComponent::AddModifierFromTemplate(UMotionModifier* Template, const UAnimSequenceBase* Animation, float StartTime, float EndTime)
+{
+	if (ensureAlways(Template))
+	{
+		FObjectDuplicationParameters Params(Template, this);
+		UMotionModifier* NewRootMotionModifier = CastChecked<UMotionModifier>(StaticDuplicateObjectEx(Params));
+		
+		NewRootMotionModifier->Animation = Animation;
+		NewRootMotionModifier->StartTime = StartTime;
+		NewRootMotionModifier->EndTime = EndTime;
+
+		AddModifier(NewRootMotionModifier);
+
+		return NewRootMotionModifier;
+	}
+
+	return nullptr;
 }
\ No newline at end of file
diff --git a/Engine/Plugins/Experimental/Animation/MotionWarping/Source/MotionWarping/Private/RootMotionModifier.cpp b/Engine/Plugins/Experimental/Animation/MotionWarping/Source/MotionWarping/Private/RootMotionModifier.cpp
--- a/Engine/Plugins/Experimental/Animation/MotionWarping/Source/MotionWarping/Private/RootMotionModifier.cpp
+++ b/Engine/Plugins/Experimental/Animation/MotionWarping/Source/MotionWarping/Private/RootMotionModifier.cpp
@@ -359,11 +359,11 @@ FRootMotionModifierHandle URootMotionModifierConfig_Warp::AddRootMotionModifierS
 {
 	if (ensureAlways(InMotionWarpingComp))
 	{
-		TSharedPtr<FRootMotionModifier_Warp> NewModifier = MakeShared<FRootMotionModifier_Warp>();
+		UMotionModifier_Warp* NewModifier = NewObject<UMotionModifier_Warp>(InMotionWarpingComp);
 		NewModifier->Animation = InAnimation;
 		NewModifier->StartTime = InStartTime;
 		NewModifier->EndTime = InEndTime;
-		NewModifier->SyncPointName = InSyncPointName;
+		NewModifier->WarpTargetName = InSyncPointName;
 		NewModifier->WarpPointAnimProvider = InWarpPointAnimProvider;
 		NewModifier->WarpPointAnimTransform = InWarpPointAnimTransform;
 		NewModifier->WarpPointAnimBoneName = InWarpPointAnimBoneName;
@@ -372,7 +372,7 @@ FRootMotionModifierHandle URootMotionModifierConfig_Warp::AddRootMotionModifierS
 		NewModifier->bWarpRotation = bInWarpRotation;
 		NewModifier->RotationType = InRotationType;
 		NewModifier->WarpRotationTimeMultiplier = InWarpRotationTimeMultiplier;
-		return InMotionWarpingComp->AddRootMotionModifier(NewModifier);
+		return InMotionWarpingComp->AddModifier(NewModifier);
 	}
 
 	return FRootMotionModifierHandle::InvalidHandle;
@@ -395,3 +395,337 @@ FRootMotionModifierHandle URootMotionModifierConfig_Scale::AddRootMotionModifier
 
 	return FRootMotionModifierHandle::InvalidHandle;
 }
+
+///////////////////////////////////////////////////////////////
+
+UMotionModifier::UMotionModifier(const FObjectInitializer& ObjectInitializer)
+	: Super(ObjectInitializer)
+{
+}
+
+void UMotionModifier::Initialize(UMotionWarpingComponent* OwnerComp)
+{
+	check(OwnerComp);
+	OwnerComponent = OwnerComp;
+	Handle.GenerateNewHandle();
+}
+
+ACharacter* UMotionModifier::GetCharacterOwner() const
+{
+	UMotionWarpingComponent* OwnerComp = GetOwnerComponent();
+	return OwnerComp ? OwnerComp->GetCharacterOwner() : nullptr;
+}
+
+void UMotionModifier::Update()
+{
+	const ACharacter* CharacterOwner = GetCharacterOwner();
+	if (CharacterOwner == nullptr)
+	{
+		return;
+	}
+
+	const FAnimMontageInstance* RootMotionMontageInstance = CharacterOwner->GetRootMotionAnimMontageInstance();
+	const UAnimMontage* Montage = RootMotionMontageInstance ? ToRawPtr(RootMotionMontageInstance->Montage) : nullptr;
+
+	// Mark for removal if our animation is not relevant anymore
+	if (Montage == nullptr || Montage != Animation)
+	{
+		UE_LOG(LogMotionWarping, Verbose, TEXT("MotionWarping: Marking RootMotionModifier for removal. Reason: Animation is not valid. Char: %s Current Montage: %s. Window: Animation: %s [%f %f] [%f %f]"),
+			*GetNameSafe(CharacterOwner), *GetNameSafe(Montage), *GetNameSafe(Animation.Get()), StartTime, EndTime, PreviousPosition, CurrentPosition);
+
+		SetState(ERootMotionModifierState::MarkedForRemoval);
+		return;
+	}
+
+	// Update playback times and weight
+	PreviousPosition = RootMotionMontageInstance->GetPreviousPosition();
+	CurrentPosition = RootMotionMontageInstance->GetPosition();
+	Weight = RootMotionMontageInstance->GetWeight();
+
+	// Mark for removal if the animation already passed the warping window
+	if (PreviousPosition >= EndTime)
+	{
+		UE_LOG(LogMotionWarping, Verbose, TEXT("MotionWarping: Marking RootMotionModifier for removal. Reason: Window has ended. Char: %s Animation: %s [%f %f] [%f %f]"),
+			*GetNameSafe(CharacterOwner), *GetNameSafe(Animation.Get()), StartTime, EndTime, PreviousPosition, CurrentPosition);
+
+		SetState(ERootMotionModifierState::MarkedForRemoval);
+		return;
+	}
+
+	// Check if we are inside the warping window
+	if (PreviousPosition >= StartTime && PreviousPosition < EndTime)
+	{
+		// If we were waiting, switch to active
+		if (GetState() == ERootMotionModifierState::Waiting)
+		{
+			SetState(ERootMotionModifierState::Active);
+		}
+	}
+
+	if (State == ERootMotionModifierState::Active)
+	{
+		if (UMotionWarpingComponent* OwnerComp = GetOwnerComponent())
+		{
+			OnUpdateDelegate.ExecuteIfBound(OwnerComp, Handle);
+		}
+	}
+}
+
+void UMotionModifier::SetState(ERootMotionModifierState NewState)
+{
+	if (State != NewState)
+	{
+		ERootMotionModifierState LastState = State;
+
+		State = NewState;
+
+		OnStateChanged(LastState);
+	}
+}
+
+void UMotionModifier::OnStateChanged(ERootMotionModifierState LastState)
+{
+	if (UMotionWarpingComponent* OwnerComp = GetOwnerComponent())
+	{
+		if (LastState != ERootMotionModifierState::Active && State == ERootMotionModifierState::Active)
+		{
+			OnActivateDelegate.ExecuteIfBound(OwnerComp, Handle);
+		}
+		else if (LastState == ERootMotionModifierState::Active && (State == ERootMotionModifierState::Disabled || State == ERootMotionModifierState::MarkedForRemoval))
+		{
+			OnDeactivateDelegate.ExecuteIfBound(OwnerComp, Handle);
+		}
+	}
+}
+
+UMotionModifier_Warp::UMotionModifier_Warp(const FObjectInitializer& ObjectInitializer)
+	: Super(ObjectInitializer)
+{
+}
+
+void UMotionModifier_Warp::Update()
+{
+	// Update playback times and state
+	Super::Update();
+
+	// Cache sync point transform and trigger OnSyncPointChanged if needed
+	const UMotionWarpingComponent* OwnerComp = GetOwnerComponent();
+	if (OwnerComp && GetState() == ERootMotionModifierState::Active)
+	{
+		const FMotionWarpingSyncPoint* SyncPointPtr = OwnerComp->FindSyncPoint(WarpTargetName);
+
+		// Disable if there is no sync point for us
+		if (SyncPointPtr == nullptr)
+		{
+			UE_LOG(LogMotionWarping, Verbose, TEXT("MotionWarping: Marking RootMotionModifier as Disabled. Reason: Invalid Sync Point (%s). Char: %s Animation: %s [%f %f] [%f %f]"),
+				*WarpTargetName.ToString(), *GetNameSafe(OwnerComp->GetOwner()), *GetNameSafe(Animation.Get()), StartTime, EndTime, PreviousPosition, CurrentPosition);
+
+			SetState(ERootMotionModifierState::Disabled);
+			return;
+		}
+
+		// Get the warp point sent by the game
+		FTransform WarpPointTransformGame = FTransform(SyncPointPtr->GetRotation(), SyncPointPtr->GetLocation());
+
+		// Initialize our target transform (where the root should end at the end of the window) with the warp point sent by the game
+		FTransform TargetTransform = WarpPointTransformGame;
+
+		// Check if a warp point is defined in the animation. If so, we need to extract it and offset the target transform 
+		// the same amount the root bone is offset from the warp point in the animation
+		if (WarpPointAnimProvider != EWarpPointAnimProvider::None)
+		{
+			if (!CachedOffsetFromWarpPoint.IsSet())
+			{
+				if (const ACharacter* CharacterOwner = GetCharacterOwner())
+				{
+					// Inverse of mesh's relative rotation. Used to convert root and warp point in the animation from Y forward to X forward
+					const FTransform MeshCompRelativeRotInverse = FTransform(CharacterOwner->GetBaseRotationOffset().Inverse());
+
+					FTransform RootTransform = FTransform::Identity;
+					FTransform WarpPointTransform = FTransform::Identity;
+					if (WarpPointAnimProvider == EWarpPointAnimProvider::Static)
+					{
+						RootTransform = MeshCompRelativeRotInverse * UMotionWarpingUtilities::ExtractRootTransformFromAnimation(GetAnimation(), EndTime);
+						WarpPointTransform = MeshCompRelativeRotInverse * WarpPointAnimTransform;
+					}
+					else if (WarpPointAnimProvider == EWarpPointAnimProvider::Bone)
+					{
+						if (const USkeletalMeshComponent* Mesh = CharacterOwner->GetMesh())
+						{
+							if (const UAnimInstance* AnimInstance = Mesh->GetAnimInstance())
+							{
+								const FBoneContainer& FullBoneContainer = AnimInstance->GetRequiredBones();
+								const int32 BoneIndex = FullBoneContainer.GetPoseBoneIndexForBoneName(WarpPointAnimBoneName);
+								if (ensure(BoneIndex != INDEX_NONE))
+								{
+									TArray<FBoneIndexType> RequiredBoneIndexArray = { 0, (FBoneIndexType)BoneIndex };
+									FullBoneContainer.GetReferenceSkeleton().EnsureParentsExistAndSort(RequiredBoneIndexArray);
+
+									FBoneContainer LimitedBoneContainer(RequiredBoneIndexArray, FCurveEvaluationOption(false), *Animation->GetSkeleton());
+
+									FCSPose<FCompactPose> Pose;
+									UMotionWarpingUtilities::ExtractComponentSpacePose(GetAnimation(), LimitedBoneContainer, EndTime, false, Pose);
+
+									RootTransform = MeshCompRelativeRotInverse * Pose.GetComponentSpaceTransform(FCompactPoseBoneIndex(0));
+									WarpPointTransform = MeshCompRelativeRotInverse * Pose.GetComponentSpaceTransform(FCompactPoseBoneIndex(1));
+								}
+							}
+						}
+					}
+
+					// Calculate and cache the offset between the root and the warp point in the animation
+					CachedOffsetFromWarpPoint = WarpPointTransform.GetRelativeTransform(RootTransform);
+				}
+			}
+
+			// Update Target Transform based on the offset between the root and the warp point in the animation
+			TargetTransform = CachedOffsetFromWarpPoint.GetValue() * WarpPointTransformGame;
+		}
+
+		if (!CachedTargetTransform.Equals(TargetTransform))
+		{
+			CachedTargetTransform = TargetTransform;
+
+			OnTargetTransformChanged();
+		}
+	}
+}
+
+FTransform UMotionModifier_Warp::ProcessRootMotion(const FTransform& InRootMotion, float DeltaSeconds)
+{
+	const ACharacter* CharacterOwner = GetCharacterOwner();
+	if (CharacterOwner == nullptr)
+	{
+		return InRootMotion;
+	}
+
+	const FTransform& CharacterTransform = CharacterOwner->GetActorTransform();
+
+	FTransform FinalRootMotion = InRootMotion;
+
+	const FTransform RootMotionTotal = UMotionWarpingUtilities::ExtractRootMotionFromAnimation(Animation.Get(), PreviousPosition, EndTime);
+
+	if (bWarpTranslation)
+	{
+		FVector DeltaTranslation = InRootMotion.GetTranslation();
+
+		const FTransform RootMotionDelta = UMotionWarpingUtilities::ExtractRootMotionFromAnimation(Animation.Get(), PreviousPosition, FMath::Min(CurrentPosition, EndTime));
+
+		const float HorizontalDelta = RootMotionDelta.GetTranslation().Size2D();
+		const float HorizontalTarget = FVector::Dist2D(CharacterTransform.GetLocation(), GetTargetLocation());
+		const float HorizontalOriginal = RootMotionTotal.GetTranslation().Size2D();
+		const float HorizontalTranslationWarped = HorizontalOriginal != 0.f ? ((HorizontalDelta * HorizontalTarget) / HorizontalOriginal) : 0.f;
+
+		if (bInLocalSpace)
+		{
+			const FTransform MeshRelativeTransform = FTransform(CharacterOwner->GetBaseRotationOffset(), CharacterOwner->GetBaseTranslationOffset());
+			const FTransform MeshTransform = MeshRelativeTransform * CharacterOwner->GetActorTransform();
+			DeltaTranslation = MeshTransform.InverseTransformPositionNoScale(GetTargetLocation()).GetSafeNormal2D() * HorizontalTranslationWarped;
+		}
+		else
+		{
+			DeltaTranslation = (GetTargetLocation() - CharacterTransform.GetLocation()).GetSafeNormal2D() * HorizontalTranslationWarped;
+		}
+
+		if (!bIgnoreZAxis)
+		{
+			const FVector CapsuleBottomLocation = (CharacterOwner->GetActorLocation() - FVector::UpVector * CharacterOwner->GetSimpleCollisionHalfHeight());
+			const float VerticalDelta = RootMotionDelta.GetTranslation().Z;
+			const float VerticalTarget = GetTargetLocation().Z - CapsuleBottomLocation.Z;
+			const float VerticalOriginal = RootMotionTotal.GetTranslation().Z;
+			const float VerticalTranslationWarped = VerticalOriginal != 0.f ? ((VerticalDelta * VerticalTarget) / VerticalOriginal) : 0.f;
+
+			DeltaTranslation.Z = VerticalTranslationWarped;
+		}
+
+		FinalRootMotion.SetTranslation(DeltaTranslation);
+	}
+
+	if (bWarpRotation)
+	{
+		const FQuat WarpedRotation = WarpRotation(InRootMotion, RootMotionTotal, DeltaSeconds);
+		FinalRootMotion.SetRotation(WarpedRotation);
+	}
+
+	// Debug
+#if !(UE_BUILD_SHIPPING || UE_BUILD_TEST)
+	const int32 DebugLevel = FMotionWarpingCVars::CVarMotionWarpingDebug.GetValueOnGameThread();
+	if (DebugLevel > 0)
+	{
+		PrintLog(TEXT("FRootMotionModifier_Simple"), InRootMotion, FinalRootMotion);
+
+		if (DebugLevel >= 2)
+		{
+			const float DrawDebugDuration = FMotionWarpingCVars::CVarMotionWarpingDrawDebugDuration.GetValueOnGameThread();
+			DrawDebugCoordinateSystem(CharacterOwner->GetWorld(), GetTargetLocation(), GetTargetRotator(), 50.f, false, DrawDebugDuration, 0, 1.f);
+		}
+	}
+#endif
+
+	return FinalRootMotion;
+}
+
+FQuat UMotionModifier_Warp::GetTargetRotation() const
+{
+	if (RotationType == EMotionWarpRotationType::Default)
+	{
+		return CachedTargetTransform.GetRotation();
+	}
+	else if (RotationType == EMotionWarpRotationType::Facing)
+	{
+		if (const ACharacter* CharacterOwner = GetCharacterOwner())
+		{
+			const FTransform& CharacterTransform = CharacterOwner->GetActorTransform();
+			const FVector ToSyncPoint = (CachedTargetTransform.GetLocation() - CharacterTransform.GetLocation()).GetSafeNormal2D();
+			return FRotationMatrix::MakeFromXZ(ToSyncPoint, FVector::UpVector).ToQuat();
+		}
+	}
+
+	return FQuat::Identity;
+}
+
+FQuat UMotionModifier_Warp::WarpRotation(const FTransform& RootMotionDelta, const FTransform& RootMotionTotal, float DeltaSeconds)
+{
+	const ACharacter* CharacterOwner = GetCharacterOwner();
+	if (CharacterOwner == nullptr)
+	{
+		return FQuat::Identity;
+	}
+
+	const FTransform& CharacterTransform = CharacterOwner->GetActorTransform();
+	const FQuat CurrentRotation = CharacterTransform.GetRotation();
+	const FQuat TargetRotation = GetTargetRotation();
+	const float TimeRemaining = (EndTime - PreviousPosition) * WarpRotationTimeMultiplier;
+	const FQuat RemainingRootRotationInWorld = RootMotionTotal.GetRotation();
+	const FQuat CurrentPlusRemainingRootMotion = RemainingRootRotationInWorld * CurrentRotation;
+	const float PercentThisStep = FMath::Clamp(DeltaSeconds / TimeRemaining, 0.f, 1.f);
+	const FQuat TargetRotThisFrame = FQuat::Slerp(CurrentPlusRemainingRootMotion, TargetRotation, PercentThisStep);
+	const FQuat DeltaOut = TargetRotThisFrame * CurrentPlusRemainingRootMotion.Inverse();
+
+	return (DeltaOut * RootMotionDelta.GetRotation());
+}
+
+#if !(UE_BUILD_SHIPPING || UE_BUILD_TEST)
+void UMotionModifier_Warp::PrintLog(const FString& Name, const FTransform& OriginalRootMotion, const FTransform& WarpedRootMotion) const
+{
+	if (const ACharacter* CharacterOwner = GetCharacterOwner())
+	{
+		const FVector CurrentLocation = (CharacterOwner->GetActorLocation() - FVector::UpVector * CharacterOwner->GetSimpleCollisionHalfHeight());
+		const FVector CurrentToTarget = (GetTargetLocation() - CurrentLocation).GetSafeNormal2D();
+		const FVector FutureLocation = CurrentLocation + (bInLocalSpace ? (CharacterOwner->GetMesh()->ConvertLocalRootMotionToWorld(WarpedRootMotion)).GetTranslation() : WarpedRootMotion.GetTranslation());
+		const FRotator CurrentRotation = CharacterOwner->GetActorRotation();
+		const FRotator FutureRotation = (WarpedRootMotion.GetRotation() * CharacterOwner->GetActorQuat()).Rotator();
+		const float Dot = FVector::DotProduct(CharacterOwner->GetActorForwardVector(), CurrentToTarget);
+		const float CurrentDist2D = FVector::Dist2D(GetTargetLocation(), CurrentLocation);
+		const float FutureDist2D = FVector::Dist2D(GetTargetLocation(), FutureLocation);
+		const float DeltaSeconds = CharacterOwner->GetWorld()->GetDeltaSeconds();
+		const float Speed = WarpedRootMotion.GetTranslation().Size() / DeltaSeconds;
+		const float EndTimeOffset = CurrentPosition - EndTime;
+
+		UE_LOG(LogMotionWarping, Log, TEXT("MotionWarping: %s. NetMode: %d Char: %s Anim: %s Win: [%f %f][%f %f] DT: %f WT: %f EndTimeOffset: %f Dist2D: %f FutureDist2D: %f Dot: %f OriginalMotionDelta: %s (%f) FinalMotionDelta: %s (%f) Speed: %f Loc: %s FutureLoc: %s Rot: %s FutureRot: %s"),
+			*Name, (int32)CharacterOwner->GetWorld()->GetNetMode(), *GetNameSafe(CharacterOwner), *GetNameSafe(Animation.Get()), StartTime, EndTime, PreviousPosition, CurrentPosition, DeltaSeconds, CharacterOwner->GetWorld()->GetTimeSeconds(), EndTimeOffset, CurrentDist2D, FutureDist2D, Dot,
+			*OriginalRootMotion.GetTranslation().ToString(), OriginalRootMotion.GetTranslation().Size(), *WarpedRootMotion.GetTranslation().ToString(), WarpedRootMotion.GetTranslation().Size(), Speed,
+			*CurrentLocation.ToString(), *FutureLocation.ToString(), *CurrentRotation.ToCompactString(), *FutureRotation.ToCompactString());
+	}
+}
+#endif
\ No newline at end of file
diff --git a/Engine/Plugins/Experimental/Animation/MotionWarping/Source/MotionWarping/Private/RootMotionModifier_AdjustmentBlendWarp.cpp b/Engine/Plugins/Experimental/Animation/MotionWarping/Source/MotionWarping/Private/RootMotionModifier_AdjustmentBlendWarp.cpp
--- a/Engine/Plugins/Experimental/Animation/MotionWarping/Source/MotionWarping/Private/RootMotionModifier_AdjustmentBlendWarp.cpp
+++ b/Engine/Plugins/Experimental/Animation/MotionWarping/Source/MotionWarping/Private/RootMotionModifier_AdjustmentBlendWarp.cpp
@@ -425,7 +425,26 @@ void URootMotionModifierConfig_AdjustmentBlendWarp::GetIKBoneTransformAndAlpha(A
 				if (Modifier->Animation == Montage)
 				{
 					static_cast<const FRootMotionModifier_AdjustmentBlendWarp*>(Modifier.Get())->GetIKBoneTransformAndAlpha(BoneName, OutTransform, OutAlpha);
-					break;
+					return;
+				}
+			}
+		}
+
+		for (UMotionModifier* Modifier : MotionWarpingComp->GetModifiers())
+		{
+			if (Modifier && Modifier->GetState() == ERootMotionModifierState::Active)
+			{
+				if(UMotionModifier_AdjustmentBlendWarp* AdjustmentBlendWarpModifier = Cast<UMotionModifier_AdjustmentBlendWarp>(Modifier))
+				{
+					// We must check if the Animation for the modifier is still relevant because in RootMotionFromMontageOnlyMode the montage could be aborted 
+					// but the modifier will remain in the list until the next update
+					const FAnimMontageInstance* RootMotionMontageInstance = Character->GetRootMotionAnimMontageInstance();
+					const UAnimMontage* Montage = RootMotionMontageInstance ? ToRawPtr(RootMotionMontageInstance->Montage) : nullptr;
+					if (Modifier->Animation == Montage)
+					{
+						AdjustmentBlendWarpModifier->GetIKBoneTransformAndAlpha(BoneName, OutTransform, OutAlpha);
+						return;
+					}
 				}
 			}
 		}
@@ -441,11 +460,11 @@ FRootMotionModifierHandle URootMotionModifierConfig_AdjustmentBlendWarp::AddRoot
 {
 	if (ensureAlways(InMotionWarpingComp))
 	{
-		TSharedPtr<FRootMotionModifier_AdjustmentBlendWarp> NewModifier = MakeShared<FRootMotionModifier_AdjustmentBlendWarp>();
+		UMotionModifier_AdjustmentBlendWarp* NewModifier = NewObject<UMotionModifier_AdjustmentBlendWarp>(InMotionWarpingComp);
 		NewModifier->Animation = InAnimation;
 		NewModifier->StartTime = InStartTime;
 		NewModifier->EndTime = InEndTime;
-		NewModifier->SyncPointName = InSyncPointName;
+		NewModifier->WarpTargetName = InSyncPointName;
 		NewModifier->WarpPointAnimProvider = InWarpPointAnimProvider;
 		NewModifier->WarpPointAnimTransform = InWarpPointAnimTransform;
 		NewModifier->WarpPointAnimBoneName = InWarpPointAnimBoneName;
@@ -454,8 +473,404 @@ FRootMotionModifierHandle URootMotionModifierConfig_AdjustmentBlendWarp::AddRoot
 		NewModifier->bWarpRotation = bInWarpRotation;
 		NewModifier->bWarpIKBones = bInWarpIKBones;
 		NewModifier->IKBones = InIKBones;
-		return InMotionWarpingComp->AddRootMotionModifier(NewModifier);
+
+		return InMotionWarpingComp->AddModifier(NewModifier);
 	}
 
 	return FRootMotionModifierHandle::InvalidHandle;
+}
+
+/////////////////////////////////////////////////////////
+
+UMotionModifier_AdjustmentBlendWarp::UMotionModifier_AdjustmentBlendWarp(const FObjectInitializer& ObjectInitializer)
+	: Super(ObjectInitializer)
+{
+	bInLocalSpace = true;
+}
+
+void UMotionModifier_AdjustmentBlendWarp::ExtractBoneTransformAtTime(FTransform& OutTransform, const FName& BoneName, float Time) const
+{
+	const int32 TrackIndex = Result.TrackNames.IndexOfByKey(BoneName);
+	ExtractBoneTransformAtTime(OutTransform, TrackIndex, Time);
+}
+
+void UMotionModifier_AdjustmentBlendWarp::ExtractBoneTransformAtTime(FTransform& OutTransform, int32 TrackIndex, float Time) const
+{
+	if (!Result.AnimationTracks.IsValidIndex(TrackIndex))
+	{
+		OutTransform = FTransform::Identity;
+		return;
+	}
+
+	const int32 TotalFrames = FMath::Max(Result.AnimationTracks[TrackIndex].PosKeys.Num(), Result.AnimationTracks[TrackIndex].RotKeys.Num());
+	const float TrackLength = (EndTime - ActualStartTime);
+	const float TimePercent = (Time - ActualStartTime) / (EndTime - ActualStartTime);
+	const float RemappedTime = TimePercent * TrackLength;
+	FAnimationUtils::ExtractTransformFromTrack(RemappedTime, TotalFrames, TrackLength, Result.AnimationTracks[TrackIndex], EAnimInterpolationType::Linear, OutTransform);
+}
+
+void UMotionModifier_AdjustmentBlendWarp::ExtractBoneTransformAtFrame(FTransform& OutTransform, int32 TrackIndex, int32 Frame) const
+{
+	if (!Result.AnimationTracks.IsValidIndex(TrackIndex) || !Result.AnimationTracks[TrackIndex].PosKeys.IsValidIndex(Frame))
+	{
+		OutTransform = FTransform::Identity;
+		return;
+	}
+
+	FAnimationUtils::ExtractTransformForFrameFromTrack(Result.AnimationTracks[TrackIndex], Frame, OutTransform);
+}
+
+void UMotionModifier_AdjustmentBlendWarp::OnTargetTransformChanged()
+{
+	const ACharacter* CharacterOwner = GetCharacterOwner();
+	const USkeletalMeshComponent* SkelMeshComp = CharacterOwner ? CharacterOwner->GetMesh() : nullptr;
+
+	if (SkelMeshComp)
+	{
+		ActualStartTime = PreviousPosition;
+		CachedRootMotion = UMotionWarpingUtilities::ExtractRootMotionFromAnimation(Animation.Get(), ActualStartTime, EndTime);
+		CachedMeshRelativeTransform = SkelMeshComp->GetRelativeTransform();
+		CachedMeshTransform = SkelMeshComp->GetComponentTransform();
+
+		Result.AnimationTracks.Reset();
+		Result.TrackNames.Reset();
+	}
+}
+
+FTransform UMotionModifier_AdjustmentBlendWarp::ProcessRootMotion(const FTransform& InRootMotion, float DeltaSeconds)
+{
+	// If warped tracks has not been generated yet, do it now
+	if (Result.GetNum() == 0)
+	{
+		PrecomputeWarpedTracks();
+	}
+
+	// Extract root motion from warped tracks
+	const FTransform FinalRootMotion = ExtractWarpedRootMotion();
+
+	// Debug
+#if !(UE_BUILD_SHIPPING || UE_BUILD_TEST)
+	if (FMotionWarpingCVars::CVarMotionWarpingDebug.GetValueOnGameThread() > 0)
+	{
+		PrintLog(TEXT("FRootMotionModifier_AdjustmentBlend"), InRootMotion, FinalRootMotion);
+
+		if (FMotionWarpingCVars::CVarMotionWarpingDebug.GetValueOnGameThread() >= 2)
+		{
+			const float DrawDebugDuration = FMotionWarpingCVars::CVarMotionWarpingDrawDebugDuration.GetValueOnGameThread();
+			DrawDebugWarpedTracks(DrawDebugDuration);
+		}
+	}
+#endif
+
+	return FinalRootMotion;
+}
+
+void UMotionModifier_AdjustmentBlendWarp::PrecomputeWarpedTracks()
+{
+	SCOPE_CYCLE_COUNTER(STAT_MotionWarping_PrecomputeWarpedTracks);
+
+	// First, extract pose at the end of the window for the bones we are going to warp
+
+	const ACharacter* CharacterOwner = GetCharacterOwner();
+	if (CharacterOwner == nullptr)
+	{
+		return;
+	}
+
+	const FBoneContainer& BoneContainer = CharacterOwner->GetMesh()->GetAnimInstance()->GetRequiredBones();
+
+	// Init FBoneContainer with only the bones that we are interested in
+	TArray<FBoneIndexType> RequiredBoneIndexArray;
+	RequiredBoneIndexArray.Add(0);
+
+	const bool bShouldWarpIKBones = bWarpIKBones && IKBones.Num() > 0;
+	if (bShouldWarpIKBones)
+	{
+		for (const FName& BoneName : IKBones)
+		{
+			const int32 BoneIndex = BoneContainer.GetPoseBoneIndexForBoneName(BoneName);
+			if (BoneIndex != INDEX_NONE)
+			{
+				RequiredBoneIndexArray.Add(BoneIndex);
+			}
+		}
+
+		BoneContainer.GetReferenceSkeleton().EnsureParentsExistAndSort(RequiredBoneIndexArray);
+	}
+
+	// Init BoneContainer
+	FBoneContainer RequiredBones(RequiredBoneIndexArray, FCurveEvaluationOption(false), *BoneContainer.GetAsset());
+
+	// Extract pose
+	FCSPose<FCompactPose> CSPose;
+	UMotionWarpingUtilities::ExtractComponentSpacePose(Animation.Get(), RequiredBones, EndTime, true, CSPose);
+
+	// Second, calculate additive pose
+
+	//Calculate additive translation for root bone
+	FVector RootTargetLocation = CachedMeshTransform.InverseTransformPositionNoScale(GetTargetLocation());
+
+	FVector RootTotalAdditiveTranslation = FVector::ZeroVector;
+	if (bWarpTranslation)
+	{
+		RootTotalAdditiveTranslation = RootTargetLocation - CachedRootMotion.GetLocation();
+
+		if (bIgnoreZAxis)
+		{
+			RootTotalAdditiveTranslation.Z = 0.f;
+		}
+	}
+
+	// Calculate additive rotation for root bone
+	FQuat RootTotalAdditiveRotation = FQuat::Identity;
+	if (bWarpRotation)
+	{
+		const FQuat TargetRotation = GetTargetRotation();
+		const FQuat OriginalRotation = CachedMeshRelativeTransform.GetRotation().Inverse() * (CachedRootMotion * CachedMeshTransform).GetRotation();
+		RootTotalAdditiveRotation = FQuat::FindBetweenNormals(OriginalRotation.GetForwardVector(), TargetRotation.GetForwardVector());
+	}
+
+	// Init Additive Pose
+	FCSPose<FCompactPose> AdditivePose;
+	AdditivePose.InitPose(&RequiredBones);
+	AdditivePose.SetComponentSpaceTransform(FCompactPoseBoneIndex(0), FTransform(RootTotalAdditiveRotation, RootTotalAdditiveTranslation));
+
+	// Calculate and add additive pose for IK bones
+	if (bShouldWarpIKBones)
+	{
+		const FTransform RootTargetPoseCS = FTransform((RootTotalAdditiveRotation * CachedRootMotion.GetRotation()), CachedRootMotion.GetTranslation() + RootTotalAdditiveTranslation);
+		for (int32 Idx = 1; Idx < CSPose.GetPose().GetNumBones(); Idx++)
+		{
+			const FName BoneName = RequiredBones.GetReferenceSkeleton().GetBoneName(RequiredBones.GetBoneIndicesArray()[Idx]);
+			if (IKBones.Contains(BoneName))
+			{
+				const int32 BoneIdx = Idx;
+				const FTransform BonePoseCS = CSPose.GetComponentSpaceTransform(FCompactPoseBoneIndex(BoneIdx));
+
+				const FTransform BoneTargetPoseCS = BonePoseCS * RootTargetPoseCS;
+				const FTransform BoneOriginalPoseCS = BonePoseCS * CachedRootMotion;
+
+				const FVector TotalAdditiveTranslation = BoneTargetPoseCS.GetLocation() - BoneOriginalPoseCS.GetLocation();
+
+				AdditivePose.SetComponentSpaceTransform(FCompactPoseBoneIndex(Idx), FTransform(TotalAdditiveTranslation));
+			}
+		}
+	}
+
+	// Finally, run adjustment blending to generate the warped poses for each bone
+
+	//@todo_fer: We could extract and cache this offline when the WarpingWindow is created
+	const float SampleRate = 1 / 60.f;
+	FMotionDeltaTrackContainer MotionDeltaTracks;
+	UMotionModifier_AdjustmentBlendWarp::ExtractMotionDeltaFromRange(RequiredBones, Animation.Get(), ActualStartTime, EndTime, SampleRate, MotionDeltaTracks);
+
+	UMotionModifier_AdjustmentBlendWarp::AdjustmentBlendWarp(RequiredBones, AdditivePose, MotionDeltaTracks, Result);
+}
+
+void UMotionModifier_AdjustmentBlendWarp::ExtractMotionDeltaFromRange(const FBoneContainer& BoneContainer, const UAnimSequenceBase* Animation, float StartTime, float EndTime, float SampleRate, FMotionDeltaTrackContainer& OutMotionDeltaTracks)
+{
+	SCOPE_CYCLE_COUNTER(STAT_MotionWarping_ExtractMotionDelta);
+
+	check(Animation);
+
+	const int32 TotalFrames = FMath::CeilToInt((EndTime - StartTime) / SampleRate) + 1;
+	const int32 TotalBones = BoneContainer.GetCompactPoseNumBones();
+
+	FCompactPose FirstFramePose;
+	UMotionWarpingUtilities::ExtractLocalSpacePose(Animation, BoneContainer, StartTime, false, FirstFramePose);
+	const FTransform RootTransformFirstFrame = FirstFramePose[FCompactPoseBoneIndex(0)];
+
+	FCSPose<FCompactPose> CSPose;
+	UMotionWarpingUtilities::ExtractComponentSpacePose(Animation, BoneContainer, StartTime, false, CSPose);
+
+	OutMotionDeltaTracks.Init(TotalBones);
+	for (int32 BoneIdx = 0; BoneIdx < TotalBones; BoneIdx++)
+	{
+		const FTransform& BoneTransform = (CSPose.GetComponentSpaceTransform(FCompactPoseBoneIndex(BoneIdx)).GetRelativeTransform(RootTransformFirstFrame));
+
+		FMotionDeltaTrack& NewTrack = OutMotionDeltaTracks.Tracks.AddDefaulted_GetRef();
+
+		NewTrack.BoneTransformTrack.Reserve(TotalFrames);
+		NewTrack.DeltaTranslationTrack.Reserve(TotalFrames);
+		NewTrack.DeltaRotationTrack.Reserve(TotalFrames);
+
+		NewTrack.BoneTransformTrack.Add(BoneTransform);
+		NewTrack.DeltaTranslationTrack.Add(FVector::ZeroVector);
+		NewTrack.DeltaRotationTrack.Add(FRotator::ZeroRotator);
+
+		NewTrack.TotalTranslation = FVector::ZeroVector;
+		NewTrack.TotalRotation = FRotator::ZeroRotator;
+	}
+
+	float Time = StartTime;
+	while (Time < EndTime)
+	{
+		Time += FMath::Min(SampleRate, EndTime - Time);
+
+		UMotionWarpingUtilities::ExtractComponentSpacePose(Animation, BoneContainer, Time, false, CSPose);
+
+		for (int32 BoneIdx = 0; BoneIdx < TotalBones; BoneIdx++)
+		{
+			FMotionDeltaTrack& Track = OutMotionDeltaTracks.Tracks[BoneIdx];
+
+			const FTransform& BoneTransform = (CSPose.GetComponentSpaceTransform(FCompactPoseBoneIndex(BoneIdx)).GetRelativeTransform(RootTransformFirstFrame));
+
+			const FTransform& LastBoneTransform = Track.BoneTransformTrack.Last();
+
+			Track.BoneTransformTrack.Add(BoneTransform);
+
+			const FVector Translation = BoneTransform.GetTranslation();
+			const FVector LastTranslation = LastBoneTransform.GetTranslation();
+			const FVector DeltaTranslation(FMath::Abs(Translation.X - LastTranslation.X), FMath::Abs(Translation.Y - LastTranslation.Y), FMath::Abs(Translation.Z - LastTranslation.Z));
+			Track.DeltaTranslationTrack.Add(DeltaTranslation);
+			Track.TotalTranslation = Track.TotalTranslation + DeltaTranslation;
+
+			const FRotator Rotation = BoneTransform.GetRotation().Rotator();
+			const FRotator LastRotation = LastBoneTransform.GetRotation().Rotator();
+			const FRotator DeltaRotation(FMath::Abs(Rotation.Pitch - LastRotation.Pitch), FMath::Abs(Rotation.Yaw - LastRotation.Yaw), FMath::Abs(Rotation.Roll - LastRotation.Roll));
+			Track.DeltaRotationTrack.Add(DeltaRotation);
+			Track.TotalRotation = Track.TotalRotation + DeltaRotation;
+		}
+	}
+}
+
+void UMotionModifier_AdjustmentBlendWarp::AdjustmentBlendWarp(const FBoneContainer& BoneContainer, const FCSPose<FCompactPose>& AdditivePose, const FMotionDeltaTrackContainer& MotionDeltaTracks, FAnimSequenceTrackContainer& Output)
+{
+	auto CalculateAdditive = [](const FVector& Total, const FVector& Delta, const FVector& Additive, const FVector& PreviousAdditive, float Alpha)
+	{
+		FVector CurrentAdditive = FVector::ZeroVector;
+		for (int32 Idx = 0; Idx < 3; Idx++)
+		{
+			if (!FMath::IsNearlyZero(Total[Idx], 1.f))
+			{
+				const float Percent = Delta[Idx] / Total[Idx];
+				const float AdditiveDelta = FMath::Abs(Additive[Idx]) * Percent;
+				CurrentAdditive[Idx] = (Additive[Idx] > 0.f) ? PreviousAdditive[Idx] + AdditiveDelta : PreviousAdditive[Idx] - AdditiveDelta;
+			}
+			else
+			{
+				CurrentAdditive[Idx] = Additive[Idx] * Alpha;
+			}
+		}
+		return CurrentAdditive;
+	};
+
+	Output.Initialize(BoneContainer.GetCompactPoseNumBones());
+
+	for (const FCompactPoseBoneIndex PoseBoneIndex : AdditivePose.GetPose().ForEachBoneIndex())
+	{
+		const FTransform& AdditiveTransform = AdditivePose.GetPose()[PoseBoneIndex];
+		if (AdditiveTransform.Equals(FTransform::Identity))
+		{
+			continue;
+		}
+
+		const int32 BoneIndex = PoseBoneIndex.GetInt();
+		const FMotionDeltaTrack& MotionDeltaTrack = MotionDeltaTracks.Tracks[BoneIndex];
+		const int32 TotalFrames = MotionDeltaTrack.BoneTransformTrack.Num();
+
+		FRawAnimSequenceTrack& Track = Output.AnimationTracks[BoneIndex];
+
+		Track.PosKeys.Reserve(TotalFrames);
+		Track.RotKeys.Reserve(TotalFrames);
+
+		Track.PosKeys.Add(MotionDeltaTrack.BoneTransformTrack[0].GetLocation());
+		Track.RotKeys.Add(MotionDeltaTrack.BoneTransformTrack[0].GetRotation());
+
+		Track.ScaleKeys.Add(FVector(1));
+
+		Output.TrackNames[BoneIndex] = BoneContainer.GetReferenceSkeleton().GetBoneName(BoneContainer.GetBoneIndicesArray()[BoneIndex]);
+
+		const FVector TotalAdditiveTranslation = AdditiveTransform.GetTranslation();
+		const FVector TotalAdditiveRotation = AdditiveTransform.GetRotation().Rotator().Euler();
+
+		FVector PrevAdditiveTranslation = FVector::ZeroVector;
+		FVector PrevAdditiveRotation = FVector::ZeroVector;
+		for (int32 FrameIdx = 1; FrameIdx < TotalFrames; FrameIdx++)
+		{
+			const FTransform& BoneTransform = MotionDeltaTrack.BoneTransformTrack[FrameIdx];
+			const float Alpha = (FrameIdx / (float)(TotalFrames - 1));
+
+			// Translation Key
+			const FVector DeltaTranslation = MotionDeltaTrack.DeltaTranslationTrack[FrameIdx];
+			const FVector CurrentAdditiveTranslation = CalculateAdditive(MotionDeltaTrack.TotalTranslation, DeltaTranslation, TotalAdditiveTranslation, PrevAdditiveTranslation, Alpha);
+			Output.AnimationTracks[BoneIndex].PosKeys.Add(BoneTransform.GetTranslation() + CurrentAdditiveTranslation);
+			PrevAdditiveTranslation = CurrentAdditiveTranslation;
+
+			// Rotation Key
+			const FVector DeltaRotation = MotionDeltaTrack.DeltaRotationTrack[FrameIdx].Euler();
+			const FVector CurrentAdditiveRotation = CalculateAdditive(MotionDeltaTrack.TotalRotation.Euler(), DeltaRotation, TotalAdditiveRotation, PrevAdditiveRotation, Alpha);
+			Output.AnimationTracks[BoneIndex].RotKeys.Add(FRotator::MakeFromEuler(BoneTransform.GetRotation().Rotator().Euler() + CurrentAdditiveRotation).Quaternion());
+			PrevAdditiveRotation = CurrentAdditiveRotation;
+		}
+	}
+}
+
+FTransform UMotionModifier_AdjustmentBlendWarp::ExtractWarpedRootMotion() const
+{
+	FTransform StartRootTransform;
+	ExtractBoneTransformAtTime(StartRootTransform, 0, PreviousPosition);
+
+	FTransform EndRootTransform;
+	ExtractBoneTransformAtTime(EndRootTransform, 0, CurrentPosition);
+
+	return EndRootTransform.GetRelativeTransform(StartRootTransform);
+}
+
+#if !(UE_BUILD_SHIPPING || UE_BUILD_TEST)
+void UMotionModifier_AdjustmentBlendWarp::DrawDebugWarpedTracks(float DrawDuration) const
+{
+	const ACharacter* CharacterOwner = GetCharacterOwner();
+	const UWorld* World = CharacterOwner ? CharacterOwner->GetWorld() : nullptr;
+	if (World && Result.GetNum() > 0 && PreviousPosition <= EndTime)
+	{
+		FTransform RootStartTransform;
+		ExtractBoneTransformAtFrame(RootStartTransform, 0, 0);
+
+		for (int32 TrackIndex = 0; TrackIndex < Result.GetNum(); TrackIndex++)
+		{
+			const int32 TotalFrames = Result.AnimationTracks[TrackIndex].PosKeys.Num();
+			if (TotalFrames > 1)
+			{
+				for (int32 FrameIdx = 0; FrameIdx < Result.AnimationTracks[TrackIndex].PosKeys.Num() - 1; FrameIdx++)
+				{
+					FTransform StartTransform;
+					ExtractBoneTransformAtFrame(StartTransform, TrackIndex, FrameIdx);
+					StartTransform = StartTransform * RootStartTransform.Inverse() * CachedMeshTransform;
+
+					FTransform EndTransform;
+					ExtractBoneTransformAtFrame(EndTransform, TrackIndex, FrameIdx + 1);
+					EndTransform = EndTransform * RootStartTransform.Inverse() * CachedMeshTransform;
+
+					DrawDebugLine(World, StartTransform.GetLocation(), EndTransform.GetLocation(), FColor::Yellow, false, DrawDuration, 0, 0.5f);
+				}
+			}
+		}
+
+		const FTransform& MeshTransform = CharacterOwner->GetMesh()->GetComponentTransform();
+		DrawDebugCoordinateSystem(World, CachedMeshTransform.GetLocation(), CachedMeshTransform.Rotator(), 20.f, false, DrawDuration, 0, 1.f);
+		DrawDebugCoordinateSystem(World, MeshTransform.GetLocation(), MeshTransform.Rotator(), 20.f, false, DrawDuration, 0, 1.f);
+		DrawDebugCoordinateSystem(World, GetTargetLocation(), GetTargetRotator(), 50.f, false, DrawDuration, 0, 1.f);
+		DrawDebugCoordinateSystem(World, (CachedRootMotion * CachedMeshTransform).GetLocation(), (CachedRootMotion * CachedMeshTransform).Rotator(), 50.f, false, DrawDuration, 0, 1.f);
+	}
+}
+#endif
+
+void UMotionModifier_AdjustmentBlendWarp::GetIKBoneTransformAndAlpha(FName BoneName, FTransform& OutTransform, float& OutAlpha) const
+{
+	if (Result.GetNum() == 0 || !bWarpIKBones || !IKBones.Contains(BoneName))
+	{
+		OutTransform = FTransform::Identity;
+		OutAlpha = 0.f;
+		return;
+	}
+
+	FTransform RootPrevPosition;
+	ExtractBoneTransformAtTime(RootPrevPosition, 0, 0.f);
+
+	FTransform BoneTransform;
+	ExtractBoneTransformAtTime(BoneTransform, BoneName, PreviousPosition);
+
+	OutTransform = BoneTransform * RootPrevPosition.Inverse() * CachedMeshTransform;
+	OutAlpha = Weight;
 }
\ No newline at end of file
diff --git a/Engine/Plugins/Experimental/Animation/MotionWarping/Source/MotionWarping/Private/RootMotionModifier_SkewWarp.cpp b/Engine/Plugins/Experimental/Animation/MotionWarping/Source/MotionWarping/Private/RootMotionModifier_SkewWarp.cpp
--- a/Engine/Plugins/Experimental/Animation/MotionWarping/Source/MotionWarping/Private/RootMotionModifier_SkewWarp.cpp
+++ b/Engine/Plugins/Experimental/Animation/MotionWarping/Source/MotionWarping/Private/RootMotionModifier_SkewWarp.cpp
@@ -172,11 +172,11 @@ FRootMotionModifierHandle URootMotionModifierConfig_SkewWarp::AddRootMotionModif
 {
 	if (ensureAlways(InMotionWarpingComp))
 	{
-		TSharedPtr<FRootMotionModifier_SkewWarp> NewModifier = MakeShared<FRootMotionModifier_SkewWarp>();
+		UMotionModifier_SkewWarp* NewModifier = NewObject<UMotionModifier_SkewWarp>(InMotionWarpingComp);
 		NewModifier->Animation = InAnimation;
 		NewModifier->StartTime = InStartTime;
 		NewModifier->EndTime = InEndTime;
-		NewModifier->SyncPointName = InSyncPointName;
+		NewModifier->WarpTargetName = InSyncPointName;
 		NewModifier->WarpPointAnimProvider = InWarpPointAnimProvider;
 		NewModifier->WarpPointAnimTransform = InWarpPointAnimTransform;
 		NewModifier->WarpPointAnimBoneName = InWarpPointAnimBoneName;
@@ -185,8 +185,171 @@ FRootMotionModifierHandle URootMotionModifierConfig_SkewWarp::AddRootMotionModif
 		NewModifier->bWarpRotation = bInWarpRotation;
 		NewModifier->RotationType = InRotationType;
 		NewModifier->WarpRotationTimeMultiplier = InWarpRotationTimeMultiplier;
-		return InMotionWarpingComp->AddRootMotionModifier(NewModifier);
+		return InMotionWarpingComp->AddModifier(NewModifier);
 	}
 
 	return FRootMotionModifierHandle::InvalidHandle;
+}
+
+///////////////////////////////////////////////////////////////
+
+UMotionModifier_SkewWarp::UMotionModifier_SkewWarp(const FObjectInitializer& ObjectInitializer)
+	: Super(ObjectInitializer)
+{
+}
+
+FTransform UMotionModifier_SkewWarp::ProcessRootMotion(const FTransform& InRootMotion, float DeltaSeconds)
+{
+	const ACharacter* CharacterOwner = GetCharacterOwner();
+	if(CharacterOwner == nullptr)
+	{
+		return InRootMotion;
+	}
+
+	FTransform FinalRootMotion = InRootMotion;
+
+	const FTransform RootMotionTotal = UMotionWarpingUtilities::ExtractRootMotionFromAnimation(Animation.Get(), PreviousPosition, EndTime);
+	const FTransform RootMotionDelta = UMotionWarpingUtilities::ExtractRootMotionFromAnimation(Animation.Get(), PreviousPosition, FMath::Min(CurrentPosition, EndTime));
+
+	if (bWarpTranslation && !RootMotionDelta.GetTranslation().IsNearlyZero())
+	{
+		const FTransform CurrentTransform = FTransform(
+			CharacterOwner->GetActorQuat(),
+			CharacterOwner->GetActorLocation() - FVector::UpVector * CharacterOwner->GetSimpleCollisionHalfHeight()
+		);
+
+		const FTransform MeshRelativeTransform = FTransform(CharacterOwner->GetBaseRotationOffset(), CharacterOwner->GetBaseTranslationOffset());
+		const FTransform MeshTransform = MeshRelativeTransform * CharacterOwner->GetActorTransform();
+		const FTransform RootMotionTotalWorldSpace = RootMotionTotal * MeshTransform;
+		const FTransform RootMotionDeltaWorldSpace = CharacterOwner->GetMesh()->ConvertLocalRootMotionToWorld(RootMotionDelta);
+
+		const FVector CurrentLocation = CurrentTransform.GetLocation();
+		const FQuat CurrentRotation = CurrentTransform.GetRotation();
+		
+		FVector TargetLocation = GetTargetLocation();
+		if(bIgnoreZAxis)
+		{
+			TargetLocation.Z = CurrentLocation.Z;
+		}
+
+		const FVector Translation = RootMotionDeltaWorldSpace.GetTranslation();
+		const FVector FutureLocation = RootMotionTotalWorldSpace.GetLocation();
+		const FVector CurrentToWorldOffset = TargetLocation - CurrentLocation;
+		const FVector CurrentToRootOffset = FutureLocation - CurrentLocation;
+
+		// Create a matrix we can use to put everything into a space looking straight at RootMotionSyncPosition. "forward" should be the axis along which we want to scale. 
+		FVector ToRootNormalized = CurrentToRootOffset.GetSafeNormal();
+
+		float BestMatchDot = FMath::Abs(FVector::DotProduct(ToRootNormalized, CurrentRotation.GetAxisX()));
+		FMatrix ToRootSyncSpace = FRotationMatrix::MakeFromXZ(ToRootNormalized, CurrentRotation.GetAxisZ());
+
+		float ZDot = FMath::Abs(FVector::DotProduct(ToRootNormalized, CurrentRotation.GetAxisZ()));
+		if (ZDot > BestMatchDot)
+		{
+			ToRootSyncSpace = FRotationMatrix::MakeFromXZ(ToRootNormalized, CurrentRotation.GetAxisX());
+			BestMatchDot = ZDot;
+		}
+
+		float YDot = FMath::Abs(FVector::DotProduct(ToRootNormalized, CurrentRotation.GetAxisY()));
+		if (YDot > BestMatchDot)
+		{
+			ToRootSyncSpace = FRotationMatrix::MakeFromXZ(ToRootNormalized, CurrentRotation.GetAxisZ());
+		}
+
+		// Put everything into RootSyncSpace.
+		const FVector RootMotionInSyncSpace = ToRootSyncSpace.InverseTransformVector(Translation);
+		const FVector CurrentToWorldSync = ToRootSyncSpace.InverseTransformVector(CurrentToWorldOffset);
+		const FVector CurrentToRootMotionSync = ToRootSyncSpace.InverseTransformVector(CurrentToRootOffset);
+
+		FVector CurrentToWorldSyncNorm = CurrentToWorldSync;
+		CurrentToWorldSyncNorm.Normalize();
+
+		FVector CurrentToRootMotionSyncNorm = CurrentToRootMotionSync;
+		CurrentToRootMotionSyncNorm.Normalize();
+
+		// Calculate skew Yaw Angle. 
+		FVector FlatToWorld = FVector(CurrentToWorldSyncNorm.X, CurrentToWorldSyncNorm.Y, 0.0f);
+		FlatToWorld.Normalize();
+		FVector FlatToRoot = FVector(CurrentToRootMotionSyncNorm.X, CurrentToRootMotionSyncNorm.Y, 0.0f);
+		FlatToRoot.Normalize();
+		float AngleAboutZ = FMath::Acos(FVector::DotProduct(FlatToWorld, FlatToRoot));
+		float AngleAboutZNorm = FMath::DegreesToRadians(FRotator::NormalizeAxis(FMath::RadiansToDegrees(AngleAboutZ)));
+		if (FlatToWorld.Y < 0.0f)
+		{
+			AngleAboutZNorm *= -1.0f;
+		}
+
+		// Calculate Skew Pitch Angle. 
+		FVector ToWorldNoY = FVector(CurrentToWorldSyncNorm.X, 0.0f, CurrentToWorldSyncNorm.Z);
+		ToWorldNoY.Normalize();
+		FVector ToRootNoY = FVector(CurrentToRootMotionSyncNorm.X, 0.0f, CurrentToRootMotionSyncNorm.Z);
+		ToRootNoY.Normalize();
+		const float AngleAboutY = FMath::Acos(FVector::DotProduct(ToWorldNoY, ToRootNoY));
+		float AngleAboutYNorm = FMath::DegreesToRadians(FRotator::NormalizeAxis(FMath::RadiansToDegrees(AngleAboutY)));
+		if (ToWorldNoY.Z < 0.0f)
+		{
+			AngleAboutYNorm *= -1.0f;
+		}
+
+		FVector SkewedRootMotion = FVector::ZeroVector;
+		float ProjectedScale = FVector::DotProduct(CurrentToWorldSync, CurrentToRootMotionSyncNorm) / CurrentToRootMotionSync.Size();
+		if (ProjectedScale != 0.0f)
+		{
+			FMatrix ScaleMatrix;
+			ScaleMatrix.SetIdentity();
+			ScaleMatrix.SetAxis(0, FVector(ProjectedScale, 0.0f, 0.0f));
+			ScaleMatrix.SetAxis(1, FVector(0.0f, 1.0f, 0.0f));
+			ScaleMatrix.SetAxis(2, FVector(0.0f, 0.0f, 1.0f));
+
+			FMatrix ShearXAlongYMatrix;
+			ShearXAlongYMatrix.SetIdentity();
+			ShearXAlongYMatrix.SetAxis(0, FVector(1.0f, FMath::Tan(AngleAboutZNorm), 0.0f));
+			ShearXAlongYMatrix.SetAxis(1, FVector(0.0f, 1.0f, 0.0f));
+			ShearXAlongYMatrix.SetAxis(2, FVector(0.0f, 0.0f, 1.0f));
+
+			FMatrix ShearXAlongZMatrix;
+			ShearXAlongZMatrix.SetIdentity();
+			ShearXAlongZMatrix.SetAxis(0, FVector(1.0f, 0.0f, FMath::Tan(AngleAboutYNorm)));
+			ShearXAlongZMatrix.SetAxis(1, FVector(0.0f, 1.0f, 0.0f));
+			ShearXAlongZMatrix.SetAxis(2, FVector(0.0f, 0.0f, 1.0f));
+
+			FMatrix ScaledSkewMatrix = ScaleMatrix * ShearXAlongYMatrix * ShearXAlongZMatrix;
+
+			// Skew and scale the Root motion. 
+			SkewedRootMotion = ScaledSkewMatrix.TransformVector(RootMotionInSyncSpace);
+		}
+		else if (!CurrentToRootMotionSync.IsZero() && !CurrentToWorldSync.IsZero() && !RootMotionInSyncSpace.IsZero())
+		{
+			// Figure out ratio between remaining Root and remaining World. Then project scaled length of current Root onto World.
+			const float Scale = CurrentToWorldSync.Size() / CurrentToRootMotionSync.Size();
+			const float StepTowardTarget = RootMotionInSyncSpace.ProjectOnTo(RootMotionInSyncSpace).Size();
+			SkewedRootMotion = CurrentToWorldSyncNorm * (Scale * StepTowardTarget);
+		}
+
+		// Put our result back in world space.  
+		FinalRootMotion.SetTranslation(ToRootSyncSpace.TransformVector(SkewedRootMotion));
+	}
+
+	if(bWarpRotation)
+	{
+		const FQuat WarpedRotation = WarpRotation(InRootMotion, RootMotionTotal, DeltaSeconds);
+		FinalRootMotion.SetRotation(WarpedRotation);
+	}
+
+	// Debug
+#if !(UE_BUILD_SHIPPING || UE_BUILD_TEST)
+	const int32 DebugLevel = FMotionWarpingCVars::CVarMotionWarpingDebug.GetValueOnGameThread();
+	if (DebugLevel > 0)
+	{
+		PrintLog(TEXT("FRootMotionModifier_Skew"), InRootMotion, FinalRootMotion);
+
+		if (DebugLevel >= 2)
+		{
+			const float DrawDebugDuration = FMotionWarpingCVars::CVarMotionWarpingDrawDebugDuration.GetValueOnGameThread();
+			DrawDebugCoordinateSystem(CharacterOwner->GetWorld(), GetTargetLocation(), GetTargetRotator(), 50.f, false, DrawDebugDuration, 0, 1.f);
+		}
+	}
+#endif
+
+	return FinalRootMotion;
 }
\ No newline at end of file
diff --git a/Engine/Plugins/Experimental/Animation/MotionWarping/Source/MotionWarping/Public/AnimNotifyState_MotionWarping.h b/Engine/Plugins/Experimental/Animation/MotionWarping/Source/MotionWarping/Public/AnimNotifyState_MotionWarping.h
--- a/Engine/Plugins/Experimental/Animation/MotionWarping/Source/MotionWarping/Public/AnimNotifyState_MotionWarping.h
+++ b/Engine/Plugins/Experimental/Animation/MotionWarping/Source/MotionWarping/Public/AnimNotifyState_MotionWarping.h
@@ -9,7 +9,6 @@
 
 class UMotionWarpingComponent;
 class UAnimSequenceBase;
-class URootMotionModifierConfig;
 
 /** AnimNotifyState used to define a motion warping window in the animation */
 UCLASS(meta = (DisplayName = "Motion Warping"))
@@ -24,6 +23,9 @@ public:
 	UPROPERTY(EditAnywhere, Instanced, BlueprintReadWrite, Category = "Config")
 	URootMotionModifierConfig* RootMotionModifierConfig;
 
+	UPROPERTY(EditAnywhere, Instanced, BlueprintReadWrite, Category = "Config")
+	UMotionModifier* RootMotionModifier;
+
 	UAnimNotifyState_MotionWarping(const FObjectInitializer& ObjectInitializer);
 
 	/** Called from the MotionWarpingComp when this notify becomes relevant. See: UMotionWarpingComponent::Update */
diff --git a/Engine/Plugins/Experimental/Animation/MotionWarping/Source/MotionWarping/Public/MotionWarpingComponent.h b/Engine/Plugins/Experimental/Animation/MotionWarping/Source/MotionWarping/Public/MotionWarpingComponent.h
--- a/Engine/Plugins/Experimental/Animation/MotionWarping/Source/MotionWarping/Public/MotionWarpingComponent.h
+++ b/Engine/Plugins/Experimental/Animation/MotionWarping/Source/MotionWarping/Public/MotionWarpingComponent.h
@@ -104,6 +104,9 @@ public:
 	/** Returns the list of root motion modifiers */
 	FORCEINLINE const TArray<TSharedPtr<FRootMotionModifier>>& GetRootMotionModifiers() const { return RootMotionModifiers; }
 
+	/** Returns the list of root motion modifiers */
+	FORCEINLINE const TArray<UMotionModifier*>& GetModifiers() const { return Modifiers; }
+
 	/** Find the SyncPoint associated with a specified name */
 	FORCEINLINE const FMotionWarpingSyncPoint* FindSyncPoint(const FName& SyncPointName) const { return SyncPoints.Find(SyncPointName); }
 
@@ -121,12 +124,20 @@ public:
 	/** Add a new modifier */
 	FRootMotionModifierHandle AddRootMotionModifier(TSharedPtr<FRootMotionModifier> Modifier);
 
+	/** Add a new modifier */
+	FRootMotionModifierHandle AddModifier(UMotionModifier* Modifier);
+
 	/** Mark all the modifiers as Disable */
 	void DisableAllRootMotionModifiers();
 
 	/** Finds a RootMotionModifier that matches the supplied Handle */
 	TSharedPtr<FRootMotionModifier> GetRootMotionModifierByHandle(const FRootMotionModifierHandle& Handle) const;
 
+	/** Finds a RootMotionModifier that matches the supplied Handle */
+	UMotionModifier* GetModifierByHandle(const FRootMotionModifierHandle& Handle) const;
+
+	UMotionModifier* AddModifierFromTemplate(UMotionModifier* Template, const UAnimSequenceBase* Animation, float StartTime, float EndTime);
+
 protected:
 
 	/** Character this component belongs to */
@@ -136,6 +147,10 @@ protected:
 	/** List of root motion modifiers */
 	TArray<TSharedPtr<FRootMotionModifier>> RootMotionModifiers;
 
+	/** List of root motion modifiers */
+	UPROPERTY(Transient)
+	TArray<UMotionModifier*> Modifiers;
+
 	UPROPERTY(Transient)
 	TMap<FName, FMotionWarpingSyncPoint> SyncPoints;
 
diff --git a/Engine/Plugins/Experimental/Animation/MotionWarping/Source/MotionWarping/Public/RootMotionModifier.h b/Engine/Plugins/Experimental/Animation/MotionWarping/Source/MotionWarping/Public/RootMotionModifier.h
--- a/Engine/Plugins/Experimental/Animation/MotionWarping/Source/MotionWarping/Public/RootMotionModifier.h
+++ b/Engine/Plugins/Experimental/Animation/MotionWarping/Source/MotionWarping/Public/RootMotionModifier.h
@@ -445,4 +445,175 @@ public:
 
 	UFUNCTION(BlueprintCallable, Category = "Motion Warping")
 	static FRootMotionModifierHandle AddRootMotionModifierScale(UMotionWarpingComponent* InMotionWarpingComp, const UAnimSequenceBase* InAnimation, float InStartTime, float InEndTime, FVector InScale);
+};
+
+///////////////////////////////////////////////////////////////
+
+UCLASS(Abstract, BlueprintType, EditInlineNew)
+class MOTIONWARPING_API UMotionModifier : public UObject
+{
+	GENERATED_BODY()
+
+public:
+
+	/** Source of the root motion we are warping */
+	UPROPERTY(EditAnywhere, BlueprintReadWrite, Category = "Config", meta = (ExposeOnSpawn))
+	TWeakObjectPtr<const UAnimSequenceBase> Animation = nullptr;
+
+	/** Start time of the warping window */
+	UPROPERTY(EditAnywhere, BlueprintReadWrite, Category = "Config", meta = (ExposeOnSpawn))
+	float StartTime = 0.f;
+
+	/** End time of the warping window */
+	UPROPERTY(EditAnywhere, BlueprintReadWrite, Category = "Config", meta = (ExposeOnSpawn))
+	float EndTime = 0.f;
+
+	/** Previous playback time of the animation */
+	UPROPERTY()
+	float PreviousPosition = 0.f;
+
+	/** Current playback time of the animation */
+	UPROPERTY()
+	float CurrentPosition = 0.f;
+
+	/** Current blend weight of the animation */
+	UPROPERTY()
+	float Weight = 0.f;
+
+	/** Whether this modifier runs before the extracted root motion is converted to world space or after */
+	UPROPERTY()
+	bool bInLocalSpace = false;
+
+	/** Delegate called when this modifier is activated (starts affecting the root motion) */
+	UPROPERTY()
+	FOnRootMotionModifierDelegate OnActivateDelegate;
+
+	/** Delegate called when this modifier updates while active (affecting the root motion) */
+	UPROPERTY()
+	FOnRootMotionModifierDelegate OnUpdateDelegate;
+
+	/** Delegate called when this modifier is deactivated (stops affecting the root motion) */
+	UPROPERTY()
+	FOnRootMotionModifierDelegate OnDeactivateDelegate;
+
+	UMotionModifier(const FObjectInitializer& ObjectInitializer);
+
+	/** Initialize this modifier */
+	void Initialize(UMotionWarpingComponent* OwnerComp);
+
+	/** Called when the state of the modifier changes */
+	virtual void OnStateChanged(ERootMotionModifierState LastState);
+
+	/** Sets the state of the modifier */
+	void SetState(ERootMotionModifierState NewState);
+
+	/** Returns the state of the modifier */
+	FORCEINLINE ERootMotionModifierState GetState() const { return State; }
+
+	/** Returns the handle to identify this modifier */
+	FORCEINLINE const FRootMotionModifierHandle& GetHandle() const { return Handle; }
+
+	/** Returns a pointer to the component that owns this modifier */
+	FORCEINLINE UMotionWarpingComponent* GetOwnerComponent() const { return OwnerComponent.Get(); }
+
+	/** Returns a pointer to the character that owns the component that owns this modifier */
+	class ACharacter* GetCharacterOwner() const;
+
+	virtual void Update();
+	virtual FTransform ProcessRootMotion(const FTransform& InRootMotion, float DeltaSeconds) { return FTransform::Identity; }
+
+	FORCEINLINE const UAnimSequenceBase* GetAnimation() const { return Animation.Get(); }
+
+private:
+
+	friend UMotionWarpingComponent;
+
+	/** Pointer back to the Component that owns this modifier */
+	UPROPERTY()
+	TWeakObjectPtr<UMotionWarpingComponent> OwnerComponent = nullptr;
+
+	/** Handle to identify this modifier */
+	UPROPERTY()
+	FRootMotionModifierHandle Handle;
+
+	/** Current state */
+	UPROPERTY()
+	ERootMotionModifierState State = ERootMotionModifierState::Waiting;
+};
+
+UCLASS()
+class MOTIONWARPING_API UMotionModifier_Warp : public UMotionModifier
+{
+	GENERATED_BODY()
+
+public:
+
+	/** Name used to find the warp target for this modifier */
+	UPROPERTY(EditAnywhere, BlueprintReadWrite, Category = "Config", meta = (ExposeOnSpawn))
+	FName WarpTargetName = NAME_None;
+
+	UPROPERTY(EditAnywhere, BlueprintReadWrite, Category = "Config")
+	EWarpPointAnimProvider WarpPointAnimProvider = EWarpPointAnimProvider::None;
+
+	//@TODO: Hide from the UI when Target != Static
+	UPROPERTY(EditAnywhere, BlueprintReadWrite, Category = "Config", meta = (EditCondition = "WarpPointAnimProvider == EWarpPointAnimProvider::Static"))
+	FTransform WarpPointAnimTransform = FTransform::Identity;
+
+	//@TODO: Hide from the UI when Target != Bone
+	UPROPERTY(EditAnywhere, BlueprintReadWrite, Category = "Config", meta = (EditCondition = "WarpPointAnimProvider == EWarpPointAnimProvider::Bone"))
+	FName WarpPointAnimBoneName = NAME_None;
+
+	/** Whether to warp the translation component of the root motion */
+	UPROPERTY(EditAnywhere, BlueprintReadWrite, Category = "Config")
+	bool bWarpTranslation = true;
+
+	/** Whether to ignore the Z component of the translation. Z motion will remain untouched */
+	UPROPERTY(EditAnywhere, BlueprintReadWrite, Category = "Config", meta = (EditCondition = "bWarpTranslation"))
+	bool bIgnoreZAxis = true;
+
+	/** Whether to warp the rotation component of the root motion */
+	UPROPERTY(EditAnywhere, BlueprintReadWrite, Category = "Config")
+	bool bWarpRotation = true;
+
+	/** Whether rotation should be warp to match the rotation of the sync point or to face the sync point */
+	UPROPERTY(EditAnywhere, BlueprintReadWrite, Category = "Config", meta = (EditCondition = "bWarpRotation"))
+	EMotionWarpRotationType RotationType;
+
+	/**
+	 * Allow to modify how fast the rotation is warped.
+	 * e.g if the window duration is 2sec and this is 0.5, the target rotation will be reached in 1sec instead of 2sec
+	 */
+	UPROPERTY(EditAnywhere, BlueprintReadWrite, Category = "Config", meta = (EditCondition = "bWarpRotation"))
+	float WarpRotationTimeMultiplier = 1.f;
+
+	UMotionModifier_Warp(const FObjectInitializer& ObjectInitializer);
+
+	//~ Begin FRootMotionModifier Interface
+	virtual void Update() override;
+	virtual FTransform ProcessRootMotion(const FTransform& InRootMotion, float DeltaSeconds) override;
+	//~ End FRootMotionModifier Interface
+
+	/** Event called during update if the target transform changes while the warping is active */
+	virtual void OnTargetTransformChanged() {}
+
+#if !(UE_BUILD_SHIPPING || UE_BUILD_TEST)
+	void PrintLog(const FString& Name, const FTransform& OriginalRootMotion, const FTransform& WarpedRootMotion) const;
+#endif
+
+protected:
+
+	FORCEINLINE FVector GetTargetLocation() const { return CachedTargetTransform.GetLocation(); }
+	FORCEINLINE FRotator GetTargetRotator() const { return GetTargetRotation().Rotator(); }
+	FQuat GetTargetRotation() const;
+
+	FQuat WarpRotation(const FTransform& RootMotionDelta, const FTransform& RootMotionTotal, float DeltaSeconds);
+
+	//@TODO: This should be Optional and private
+	UPROPERTY()
+	FTransform CachedTargetTransform = FTransform::Identity;
+
+private:
+
+	/** Cached of the offset from the warp target. Used to calculate the final target transform when a warp target is defined in the animation */
+	TOptional<FTransform> CachedOffsetFromWarpPoint;
 };
\ No newline at end of file
diff --git a/Engine/Plugins/Experimental/Animation/MotionWarping/Source/MotionWarping/Public/RootMotionModifier_AdjustmentBlendWarp.h b/Engine/Plugins/Experimental/Animation/MotionWarping/Source/MotionWarping/Public/RootMotionModifier_AdjustmentBlendWarp.h
--- a/Engine/Plugins/Experimental/Animation/MotionWarping/Source/MotionWarping/Public/RootMotionModifier_AdjustmentBlendWarp.h
+++ b/Engine/Plugins/Experimental/Animation/MotionWarping/Source/MotionWarping/Public/RootMotionModifier_AdjustmentBlendWarp.h
@@ -127,4 +127,60 @@ public:
 
 	UFUNCTION(BlueprintPure, Category = "Motion Warping")
 	static void GetIKBoneTransformAndAlpha(ACharacter* Character, FName BoneName, FTransform& OutTransform, float& OutAlpha);
+};
+
+//////////////////////////////////////////////
+
+UCLASS()
+class MOTIONWARPING_API UMotionModifier_AdjustmentBlendWarp : public UMotionModifier_Warp
+{
+	GENERATED_BODY()
+
+public:
+
+	UPROPERTY(EditAnywhere, BlueprintReadWrite, Category = "Config")
+	bool bWarpIKBones = false;
+
+	UPROPERTY(EditAnywhere, BlueprintReadWrite, Category = "Config")
+	TArray<FName> IKBones;
+
+	UMotionModifier_AdjustmentBlendWarp(const FObjectInitializer& ObjectInitializer);
+
+	virtual void OnTargetTransformChanged() override;
+	virtual FTransform ProcessRootMotion(const FTransform& InRootMotion, float DeltaSeconds) override;
+
+	void GetIKBoneTransformAndAlpha(FName BoneName, FTransform& OutTransform, float& OutAlpha) const;
+
+protected:
+
+	UPROPERTY()
+	FTransform CachedMeshTransform;
+
+	UPROPERTY()
+	FTransform CachedMeshRelativeTransform;
+
+	UPROPERTY()
+	FTransform CachedRootMotion;
+
+	UPROPERTY()
+	FAnimSequenceTrackContainer Result;
+
+	UPROPERTY()
+	float ActualStartTime = 0.f;
+
+	void PrecomputeWarpedTracks();
+
+	FTransform ExtractWarpedRootMotion() const;
+
+	void ExtractBoneTransformAtTime(FTransform& OutTransform, const FName& BoneName, float Time) const;
+	void ExtractBoneTransformAtTime(FTransform& OutTransform, int32 TrackIndex, float Time) const;
+	void ExtractBoneTransformAtFrame(FTransform& OutTransform, int32 TrackIndex, int32 Frame) const;
+
+#if !(UE_BUILD_SHIPPING || UE_BUILD_TEST)
+	void DrawDebugWarpedTracks(float DrawDuration) const;
+#endif
+
+	static void ExtractMotionDeltaFromRange(const FBoneContainer& BoneContainer, const UAnimSequenceBase* Animation, float StartTime, float EndTime, float SampleRate, FMotionDeltaTrackContainer& OutMotionDeltaTracks);
+
+	static void AdjustmentBlendWarp(const FBoneContainer& BoneContainer, const FCSPose<FCompactPose>& AdditivePose, const FMotionDeltaTrackContainer& MotionDeltaTracks, FAnimSequenceTrackContainer& Output);
 };
\ No newline at end of file
diff --git a/Engine/Plugins/Experimental/Animation/MotionWarping/Source/MotionWarping/Public/RootMotionModifier_SkewWarp.h b/Engine/Plugins/Experimental/Animation/MotionWarping/Source/MotionWarping/Public/RootMotionModifier_SkewWarp.h
--- a/Engine/Plugins/Experimental/Animation/MotionWarping/Source/MotionWarping/Public/RootMotionModifier_SkewWarp.h
+++ b/Engine/Plugins/Experimental/Animation/MotionWarping/Source/MotionWarping/Public/RootMotionModifier_SkewWarp.h
@@ -41,4 +41,18 @@ public:
 	static FRootMotionModifierHandle AddRootMotionModifierSkewWarp(UMotionWarpingComponent* InMotionWarpingComp, const UAnimSequenceBase* InAnimation, float InStartTime, float InEndTime, 
 		FName InSyncPointName, EWarpPointAnimProvider InWarpPointAnimProvider, FTransform InWarpPointAnimTransform, FName InWarpPointAnimBoneName,
 		bool bInWarpTranslation, bool bInIgnoreZAxis, bool bInWarpRotation, EMotionWarpRotationType InRotationType, float InWarpRotationTimeMultiplier = 1.f);
+};
+
+//////////////////////////////////////////////
+
+UCLASS()
+class MOTIONWARPING_API UMotionModifier_SkewWarp : public UMotionModifier_Warp
+{
+	GENERATED_BODY()
+
+public:
+
+	UMotionModifier_SkewWarp(const FObjectInitializer& ObjectInitializer);
+
+	virtual FTransform ProcessRootMotion(const FTransform& InRootMotion, float DeltaSeconds) override;
 };
\ No newline at end of file
