Motion Warping - Removed deprecated ustruct version of motion modifiers

diff --git a/Engine/Plugins/Experimental/Animation/ContextualAnimation/Source/ContextualAnimation/Private/ContextualAnimTypes.cpp b/Engine/Plugins/Experimental/Animation/ContextualAnimation/Source/ContextualAnimation/Private/ContextualAnimTypes.cpp
--- a/Engine/Plugins/Experimental/Animation/ContextualAnimation/Source/ContextualAnimation/Private/ContextualAnimTypes.cpp
+++ b/Engine/Plugins/Experimental/Animation/ContextualAnimation/Source/ContextualAnimation/Private/ContextualAnimTypes.cpp
@@ -49,7 +49,7 @@ float FContextualAnimData::GetSyncTimeForWarpSection(int32 WarpSectionIndex) con
 
 	if(Animation && WarpSectionIndex >= 0)
 	{
-		FName LastSyncPointName = NAME_None;
+		FName LastWarpTargetName = NAME_None;
 		int32 LastWarpSectionIndex = INDEX_NONE;
 
 		for (int32 Idx = 0; Idx < Animation->Notifies.Num(); Idx++)
@@ -57,22 +57,21 @@ float FContextualAnimData::GetSyncTimeForWarpSection(int32 WarpSectionIndex) con
 			const FAnimNotifyEvent& NotifyEvent = Animation->Notifies[Idx];
 			if(const UAnimNotifyState_MotionWarping* Notify = Cast<const UAnimNotifyState_MotionWarping>(NotifyEvent.NotifyStateClass))
 			{
-				//@TODO: We may want to add a function in UAnimNotifyState_MotionWarping that returns the SyncPointName?
-				if(const URootMotionModifierConfig_Warp* Config = Cast<const URootMotionModifierConfig_Warp>(Notify->RootMotionModifierConfig))
+				if(const URootMotionModifier_Warp* Modifier = Cast<const URootMotionModifier_Warp>(Notify->RootMotionModifier))
 				{
-					const FName SyncPointName = Config->SyncPointName;
-					if(SyncPointName != NAME_None)
+					const FName WarpTargetName = Modifier->WarpTargetName;
+					if(WarpTargetName != NAME_None)
 					{
 						// First valid warping window. Initialize everything
 						if (LastWarpSectionIndex == INDEX_NONE)
 						{
-							LastSyncPointName = SyncPointName;
+							LastWarpTargetName = WarpTargetName;
 							Result = NotifyEvent.GetEndTriggerTime();
 							LastWarpSectionIndex = 0;
 						}
 						// If we hit another warping window but the sync point is the same as the previous, update SyncTime.
 						// This is to deal with cases where a first short window is used to face the alignment point and a second one to perform the rest of the warp
-						else if(SyncPointName == LastSyncPointName)
+						else if(WarpTargetName == LastWarpTargetName)
 						{
 							Result = NotifyEvent.GetEndTriggerTime();
 						}
@@ -82,7 +81,7 @@ float FContextualAnimData::GetSyncTimeForWarpSection(int32 WarpSectionIndex) con
 							// If we haven't reached the desired WarpSection yet. Update control vars and keep moving
 							if(WarpSectionIndex > LastWarpSectionIndex)
 							{
-								LastSyncPointName = SyncPointName;
+								LastWarpTargetName = WarpTargetName;
 								Result = NotifyEvent.GetEndTriggerTime();
 								LastWarpSectionIndex++;
 							}
@@ -103,7 +102,7 @@ float FContextualAnimData::GetSyncTimeForWarpSection(int32 WarpSectionIndex) con
 
 float FContextualAnimData::GetSyncTimeForWarpSection(const FName& WarpSectionName) const
 {
-	//@TODO: We need a better way to identify warping sections withing the animation. This is just a temp solution
+	//@TODO: We need a better way to identify warping sections within the animation. This is just a temp solution
 	//@TODO: We should cache this data
 
 	float Result = 0.f;
@@ -115,11 +114,10 @@ float FContextualAnimData::GetSyncTimeForWarpSection(const FName& WarpSectionNam
 			const FAnimNotifyEvent& NotifyEvent = Animation->Notifies[Idx];
 			if (const UAnimNotifyState_MotionWarping* Notify = Cast<const UAnimNotifyState_MotionWarping>(NotifyEvent.NotifyStateClass))
 			{
-				//@TODO: We may want to add a function in UAnimNotifyState_MotionWarping that returns the SyncPointName?
-				if (const URootMotionModifierConfig_Warp* Config = Cast<const URootMotionModifierConfig_Warp>(Notify->RootMotionModifierConfig))
+				if (const URootMotionModifier_Warp* Config = Cast<const URootMotionModifier_Warp>(Notify->RootMotionModifier))
 				{
-					const FName SyncPointName = Config->SyncPointName;
-					if (WarpSectionName == SyncPointName)
+					const FName WarpTargetName = Config->WarpTargetName;
+					if (WarpSectionName == WarpTargetName)
 					{
 						const float NotifyEndTriggerTime = NotifyEvent.GetEndTriggerTime();
 						if(NotifyEndTriggerTime > Result)
diff --git a/Engine/Plugins/Experimental/Animation/MotionWarping/Source/MotionWarping/Private/AnimNotifyState_MotionWarping.cpp b/Engine/Plugins/Experimental/Animation/MotionWarping/Source/MotionWarping/Private/AnimNotifyState_MotionWarping.cpp
--- a/Engine/Plugins/Experimental/Animation/MotionWarping/Source/MotionWarping/Private/AnimNotifyState_MotionWarping.cpp
+++ b/Engine/Plugins/Experimental/Animation/MotionWarping/Source/MotionWarping/Private/AnimNotifyState_MotionWarping.cpp
@@ -16,50 +16,17 @@ FRootMotionModifierHandle UAnimNotifyState_MotionWarping::OnBecomeRelevant(UMoti
 	//@TODO: Temp, to handle the case where the blueprint notify does not return the handle. Will be removed shortly
 	if (!Handle.IsValid())
 	{
-		if (MotionWarpingComp->GetRootMotionModifiers().Num())
+		if (MotionWarpingComp->GetModifiers().Num())
 		{
-			TSharedPtr<FRootMotionModifier> Last = MotionWarpingComp->GetRootMotionModifiers().Last();
-			if (Last.IsValid() && Last->Animation.Get() == Animation && Last->StartTime == StartTime && Last->EndTime == EndTime)
+			URootMotionModifier* Last = MotionWarpingComp->GetModifiers().Last();
+			if (Last && Last->Animation.Get() == Animation && Last->StartTime == StartTime && Last->EndTime == EndTime)
 			{
 				Handle = Last->GetHandle();
 			}
 		}
-
-		if (!Handle.IsValid())
-		{
-			if (MotionWarpingComp->GetModifiers().Num())
-			{
-				UMotionModifier* Last = MotionWarpingComp->GetModifiers().Last();
-				if (Last && Last->Animation.Get() == Animation && Last->StartTime == StartTime && Last->EndTime == EndTime)
-				{
-					Handle = Last->GetHandle();
-				}
-			}
-		}
-	}
-
-	TSharedPtr<FRootMotionModifier> Modifier = MotionWarpingComp->GetRootMotionModifierByHandle(Handle);
-	if (Modifier.IsValid())
-	{
-		if (!Modifier->OnActivateDelegate.IsBound())
-		{
-			Modifier->OnActivateDelegate.BindDynamic(this, &UAnimNotifyState_MotionWarping::OnRootMotionModifierActivate);
-		}
-
-		if (!Modifier->OnUpdateDelegate.IsBound())
-		{
-			Modifier->OnUpdateDelegate.BindDynamic(this, &UAnimNotifyState_MotionWarping::OnRootMotionModifierUpdate);
-		}
-
-		if (!Modifier->OnDeactivateDelegate.IsBound())
-		{
-			Modifier->OnDeactivateDelegate.BindDynamic(this, &UAnimNotifyState_MotionWarping::OnRootMotionModifierDeactivate);
-		}
-
-		return Handle;
 	}
 
-	if (UMotionModifier* RootMotionModifierNew = MotionWarpingComp->GetModifierByHandle(Handle))
+	if (URootMotionModifier* RootMotionModifierNew = MotionWarpingComp->GetModifierByHandle(Handle))
 	{
 		if (!RootMotionModifierNew->OnActivateDelegate.IsBound())
 		{
@@ -86,13 +53,9 @@ FRootMotionModifierHandle UAnimNotifyState_MotionWarping::AddRootMotionModifier_
 {
 	if (MotionWarpingComp)
 	{
-		if (RootMotionModifierConfig)
-		{
-			return RootMotionModifierConfig->AddRootMotionModifierNew(MotionWarpingComp, Animation, StartTime, EndTime);
-		}
-		else if (RootMotionModifier)
+		if (RootMotionModifier)
 		{
-			UMotionModifier* NewRootMotionModifier = MotionWarpingComp->AddModifierFromTemplate(RootMotionModifier, Animation, StartTime, EndTime);
+			URootMotionModifier* NewRootMotionModifier = MotionWarpingComp->AddModifierFromTemplate(RootMotionModifier, Animation, StartTime, EndTime);
 			if (NewRootMotionModifier)
 			{
 				return NewRootMotionModifier->GetHandle();
diff --git a/Engine/Plugins/Experimental/Animation/MotionWarping/Source/MotionWarping/Private/MotionWarpingComponent.cpp b/Engine/Plugins/Experimental/Animation/MotionWarping/Source/MotionWarping/Private/MotionWarpingComponent.cpp
--- a/Engine/Plugins/Experimental/Animation/MotionWarping/Source/MotionWarping/Private/MotionWarpingComponent.cpp
+++ b/Engine/Plugins/Experimental/Animation/MotionWarping/Source/MotionWarping/Private/MotionWarpingComponent.cpp
@@ -120,9 +120,9 @@ void UMotionWarpingUtilities::GetMotionWarpingWindowsFromAnimation(const UAnimSe
 	}
 }
 
-void UMotionWarpingUtilities::GetMotionWarpingWindowsForSyncPointFromAnimation(const UAnimSequenceBase* Animation, FName SyncPointName, TArray<FMotionWarpingWindowData>& OutWindows)
+void UMotionWarpingUtilities::GetMotionWarpingWindowsForSyncPointFromAnimation(const UAnimSequenceBase* Animation, FName WarpTargetName, TArray<FMotionWarpingWindowData>& OutWindows)
 {
-	if (Animation && SyncPointName != NAME_None)
+	if (Animation && WarpTargetName != NAME_None)
 	{
 		OutWindows.Reset();
 
@@ -131,9 +131,9 @@ void UMotionWarpingUtilities::GetMotionWarpingWindowsForSyncPointFromAnimation(c
 			const FAnimNotifyEvent& NotifyEvent = Animation->Notifies[Idx];
 			if (UAnimNotifyState_MotionWarping* Notify = Cast<UAnimNotifyState_MotionWarping>(NotifyEvent.NotifyStateClass))
 			{
-				if (const URootMotionModifierConfig_Warp* Config = Cast<const URootMotionModifierConfig_Warp>(Notify->RootMotionModifierConfig))
+				if (const URootMotionModifier_Warp* Modifier = Cast<const URootMotionModifier_Warp>(Notify->RootMotionModifier))
 				{
-					if(Config->SyncPointName == SyncPointName)
+					if(Modifier->WarpTargetName == WarpTargetName)
 					{
 						FMotionWarpingWindowData Data;
 						Data.AnimNotify = Notify;
@@ -172,41 +172,13 @@ void UMotionWarpingComponent::InitializeComponent()
 
 bool UMotionWarpingComponent::ContainsModifier(const UAnimSequenceBase* Animation, float StartTime, float EndTime) const
 {
-	// DEPRECATED will be removed shortly
-	if (RootMotionModifiers.ContainsByPredicate([=](const TSharedPtr<FRootMotionModifier>& Modifier)
-		{
-			return (Modifier.IsValid() && Modifier->Animation == Animation && Modifier->StartTime == StartTime && Modifier->EndTime == EndTime);
-		}))
-	{
-		return true;
-	}
-
-	return Modifiers.ContainsByPredicate([=](const UMotionModifier* Modifier)
+	return Modifiers.ContainsByPredicate([=](const URootMotionModifier* Modifier)
 		{
 			return (Modifier->Animation == Animation && Modifier->StartTime == StartTime && Modifier->EndTime == EndTime);
 		});
 }
 
-FRootMotionModifierHandle UMotionWarpingComponent::AddRootMotionModifier(TSharedPtr<FRootMotionModifier> Modifier)
-{
-	if (ensureAlways(Modifier.IsValid()))
-	{
-		//@TODO: Move this from here after update all modifiers
-		Modifier->Initialize(this);
-
-		RootMotionModifiers.Add(Modifier);
-
-		UE_LOG(LogMotionWarping, Verbose, TEXT("MotionWarping: RootMotionModifier added. NetMode: %d WorldTime: %f Char: %s Animation: %s [%f %f] [%f %f] Loc: %s Rot: %s"),
-			GetWorld()->GetNetMode(), GetWorld()->GetTimeSeconds(), *GetNameSafe(GetCharacterOwner()), *GetNameSafe(Modifier->Animation.Get()), Modifier->StartTime, Modifier->EndTime, Modifier->PreviousPosition, Modifier->CurrentPosition,
-			*GetCharacterOwner()->GetActorLocation().ToString(), *GetCharacterOwner()->GetActorRotation().ToCompactString());
-
-		return Modifier->GetHandle();
-	}
-
-	return FRootMotionModifierHandle::InvalidHandle;
-}
-
-FRootMotionModifierHandle UMotionWarpingComponent::AddModifier(UMotionModifier* Modifier)
+FRootMotionModifierHandle UMotionWarpingComponent::AddModifier(URootMotionModifier* Modifier)
 {
 	if (ensureAlways(Modifier))
 	{
@@ -225,27 +197,11 @@ FRootMotionModifierHandle UMotionWarpingComponent::AddModifier(UMotionModifier*
 	return FRootMotionModifierHandle::InvalidHandle;
 }
 
-TSharedPtr<FRootMotionModifier> UMotionWarpingComponent::GetRootMotionModifierByHandle(const FRootMotionModifierHandle& Handle) const
-{
-	if(Handle.IsValid())
-	{
-		for (const TSharedPtr<FRootMotionModifier>& RootMotionModifier : RootMotionModifiers)
-		{
-			if (RootMotionModifier.IsValid() && RootMotionModifier->GetHandle() == Handle)
-			{
-				return RootMotionModifier;
-			}
-		}
-	}
-
-	return TSharedPtr<FRootMotionModifier>(nullptr);
-}
-
-UMotionModifier* UMotionWarpingComponent::GetModifierByHandle(const FRootMotionModifierHandle& Handle) const
+URootMotionModifier* UMotionWarpingComponent::GetModifierByHandle(const FRootMotionModifierHandle& Handle) const
 {
 	if (Handle.IsValid())
 	{
-		for (UMotionModifier* RootMotionModifier : Modifiers)
+		for (URootMotionModifier* RootMotionModifier : Modifiers)
 		{
 			if (RootMotionModifier && RootMotionModifier->GetHandle() == Handle)
 			{
@@ -259,18 +215,9 @@ UMotionModifier* UMotionWarpingComponent::GetModifierByHandle(const FRootMotionM
 
 void UMotionWarpingComponent::DisableAllRootMotionModifiers()
 {
-	// DEPRECATED will be removed shortly
-	if (RootMotionModifiers.Num() > 0)
-	{
-		for (TSharedPtr<FRootMotionModifier>& Modifier : RootMotionModifiers)
-		{
-			Modifier->SetState(ERootMotionModifierState::Disabled);
-		}
-	}
-
 	if (Modifiers.Num() > 0)
 	{
-		for (UMotionModifier* Modifier : Modifiers)
+		for (URootMotionModifier* Modifier : Modifiers)
 		{
 			Modifier->SetState(ERootMotionModifierState::Disabled);
 		}
@@ -344,41 +291,16 @@ void UMotionWarpingComponent::Update()
 
 	OnPreUpdate.Broadcast(this);
 
-	// Update the state of all the modifiers
-	if(RootMotionModifiers.Num() > 0)
-	{
-		for (TSharedPtr<FRootMotionModifier>& Modifier : RootMotionModifiers)
-		{
-			//@TODO: Replace with the signature that doesn't receive the comp via param
-			Modifier->Update(*this);
-		}
-
-		// Remove the modifiers that has been marked for removal
-		RootMotionModifiers.RemoveAll([this](const TSharedPtr<FRootMotionModifier>& Modifier) 
-		{ 
-			if(Modifier->GetState() == ERootMotionModifierState::MarkedForRemoval)
-			{
-				UE_LOG(LogMotionWarping, Verbose, TEXT("MotionWarping: RootMotionModifier removed. NetMode: %d WorldTime: %f Char: %s Animation: %s [%f %f] [%f %f] Loc: %s Rot: %s"),
-					GetWorld()->GetNetMode(), GetWorld()->GetTimeSeconds(), *GetNameSafe(GetCharacterOwner()), *GetNameSafe(Modifier->Animation.Get()), Modifier->StartTime, Modifier->EndTime, Modifier->PreviousPosition, Modifier->CurrentPosition,
-					*GetCharacterOwner()->GetActorLocation().ToString(), *GetCharacterOwner()->GetActorRotation().ToCompactString());
-
-				return true;
-			}
-
-			return false; 
-		});
-	}
-
 	// Update the state of all the modifiers
 	if (Modifiers.Num() > 0)
 	{
-		for (UMotionModifier* Modifier : Modifiers)
+		for (URootMotionModifier* Modifier : Modifiers)
 		{
 			Modifier->Update();
 		}
 
 		// Remove the modifiers that has been marked for removal
-		Modifiers.RemoveAll([this](const UMotionModifier* Modifier)
+		Modifiers.RemoveAll([this](const URootMotionModifier* Modifier)
 		{
 			if (Modifier->GetState() == ERootMotionModifierState::MarkedForRemoval)
 			{
@@ -409,17 +331,7 @@ FTransform UMotionWarpingComponent::ProcessRootMotionPreConvertToWorld(const FTr
 	FTransform FinalRootMotion = InRootMotion;
 
 	// Apply Local Space Modifiers
-	for (TSharedPtr<FRootMotionModifier>& RootMotionModifier : RootMotionModifiers)
-	{
-		if (RootMotionModifier->GetState() == ERootMotionModifierState::Active && RootMotionModifier->bInLocalSpace)
-		{
-			//@TODO: Replace with the signature that doesn't receive the comp via param
-			FinalRootMotion = RootMotionModifier->ProcessRootMotion(*this, FinalRootMotion, DeltaSeconds);
-		}
-	}
-
-	// Apply Local Space Modifiers
-	for (UMotionModifier* Modifier : Modifiers)
+	for (URootMotionModifier* Modifier : Modifiers)
 	{
 		if (Modifier->GetState() == ERootMotionModifierState::Active && Modifier->bInLocalSpace)
 		{
@@ -442,17 +354,7 @@ FTransform UMotionWarpingComponent::ProcessRootMotionPostConvertToWorld(const FT
 	FTransform FinalRootMotion = InRootMotion;
 
 	// Apply World Space Modifiers
-	for (TSharedPtr<FRootMotionModifier>& RootMotionModifier : RootMotionModifiers)
-	{
-		if (RootMotionModifier->GetState() == ERootMotionModifierState::Active && !RootMotionModifier->bInLocalSpace)
-		{
-			//@TODO: Replace with the signature that doesn't receive the comp via param
-			FinalRootMotion = RootMotionModifier->ProcessRootMotion(*this, FinalRootMotion, DeltaSeconds);
-		}
-	}
-
-	// Apply World Space Modifiers
-	for (UMotionModifier* Modifier : Modifiers)
+	for (URootMotionModifier* Modifier : Modifiers)
 	{
 		if (Modifier->GetState() == ERootMotionModifierState::Active && !Modifier->bInLocalSpace)
 		{
@@ -467,7 +369,7 @@ FTransform UMotionWarpingComponent::ProcessRootMotionPostConvertToWorld(const FT
 		const float DrawDebugDuration = FMotionWarpingCVars::CVarMotionWarpingDrawDebugDuration.GetValueOnGameThread();
 		const float PointSize = 7.f;
 		const FVector ActorFeetLocation = CharacterMovementComponent->GetActorFeetLocation();
-		if (RootMotionModifiers.Num() > 0)
+		if (Modifiers.Num() > 0)
 		{
 			if (!OriginalRootMotionAccum.IsSet())
 			{
@@ -510,17 +412,7 @@ int32 UMotionWarpingComponent::RemoveSyncPoint(FName Name)
 
 bool UMotionWarpingComponent::SetRootMotionModifierDelegatesByHandle(FRootMotionModifierHandle Handle, FOnRootMotionModifierDelegate OnActivate, FOnRootMotionModifierDelegate OnUpdate, FOnRootMotionModifierDelegate OnDeactivate)
 {
-	// DEPRECATED will be removed shortly
-	TSharedPtr<FRootMotionModifier> Modifier = GetRootMotionModifierByHandle(Handle);
-	if (Modifier.IsValid())
-	{
-		Modifier->OnActivateDelegate = OnActivate;
-		Modifier->OnUpdateDelegate = OnUpdate;
-		Modifier->OnDeactivateDelegate = OnDeactivate;
-		return true;
-	}
-
-	if (UMotionModifier* RootMotionModifier = GetModifierByHandle(Handle))
+	if (URootMotionModifier* RootMotionModifier = GetModifierByHandle(Handle))
 	{
 		RootMotionModifier->OnActivateDelegate = OnActivate;
 		RootMotionModifier->OnUpdateDelegate = OnUpdate;
@@ -533,17 +425,7 @@ bool UMotionWarpingComponent::SetRootMotionModifierDelegatesByHandle(FRootMotion
 
 bool UMotionWarpingComponent::GetAnimationAndTimeRangeForRootMotionModifierByHandle(FRootMotionModifierHandle Handle, UAnimSequenceBase*& OutAnimation, float& OutStartTime, float& OutEndTime) const
 {
-	// DEPRECATED will be removed shortly
-	TSharedPtr<FRootMotionModifier> Modifier = GetRootMotionModifierByHandle(Handle);
-	if (Modifier.IsValid())
-	{
-		OutAnimation = const_cast<UAnimSequenceBase*>(Modifier->GetAnimation());
-		OutStartTime = Modifier->StartTime;
-		OutEndTime = Modifier->EndTime;
-		return true;
-	}
-
-	if (UMotionModifier* RootMotionModifier = GetModifierByHandle(Handle))
+	if (URootMotionModifier* RootMotionModifier = GetModifierByHandle(Handle))
 	{
 		OutAnimation = const_cast<UAnimSequenceBase*>(RootMotionModifier->GetAnimation());
 		OutStartTime = RootMotionModifier->StartTime;
@@ -554,12 +436,12 @@ bool UMotionWarpingComponent::GetAnimationAndTimeRangeForRootMotionModifierByHan
 	return false;
 }
 
-UMotionModifier* UMotionWarpingComponent::AddModifierFromTemplate(UMotionModifier* Template, const UAnimSequenceBase* Animation, float StartTime, float EndTime)
+URootMotionModifier* UMotionWarpingComponent::AddModifierFromTemplate(URootMotionModifier* Template, const UAnimSequenceBase* Animation, float StartTime, float EndTime)
 {
 	if (ensureAlways(Template))
 	{
 		FObjectDuplicationParameters Params(Template, this);
-		UMotionModifier* NewRootMotionModifier = CastChecked<UMotionModifier>(StaticDuplicateObjectEx(Params));
+		URootMotionModifier* NewRootMotionModifier = CastChecked<URootMotionModifier>(StaticDuplicateObjectEx(Params));
 		
 		NewRootMotionModifier->Animation = Animation;
 		NewRootMotionModifier->StartTime = StartTime;
diff --git a/Engine/Plugins/Experimental/Animation/MotionWarping/Source/MotionWarping/Private/RootMotionModifier.cpp b/Engine/Plugins/Experimental/Animation/MotionWarping/Source/MotionWarping/Private/RootMotionModifier.cpp
--- a/Engine/Plugins/Experimental/Animation/MotionWarping/Source/MotionWarping/Private/RootMotionModifier.cpp
+++ b/Engine/Plugins/Experimental/Animation/MotionWarping/Source/MotionWarping/Private/RootMotionModifier.cpp
@@ -19,404 +19,25 @@ void FRootMotionModifierHandle::GenerateNewHandle()
 	Handle = GHandle++;
 }
 
-// FRootMotionModifier
-///////////////////////////////////////////////////////////////
-
-void FRootMotionModifier::Initialize(UMotionWarpingComponent* OwnerComp)
-{
-	check(OwnerComp);
-	OwnerComponent = OwnerComp;
-	Handle.GenerateNewHandle();
-}
-
-ACharacter* FRootMotionModifier::GetCharacterOwner() const
-{
-	UMotionWarpingComponent* OwnerComp = GetOwnerComponent();
-	return OwnerComp ? OwnerComp->GetCharacterOwner() : nullptr;
-}
-
-void FRootMotionModifier::Update(UMotionWarpingComponent& OwnerComp_DEPRECATED)
-{
-	const ACharacter* CharacterOwner = GetCharacterOwner();
-	if(CharacterOwner == nullptr)
-	{
-		return;
-	}
-
-	const FAnimMontageInstance* RootMotionMontageInstance = CharacterOwner->GetRootMotionAnimMontageInstance();
-	const UAnimMontage* Montage = RootMotionMontageInstance ? ToRawPtr(RootMotionMontageInstance->Montage) : nullptr;
-
-	// Mark for removal if our animation is not relevant anymore
-	if (Montage == nullptr || Montage != Animation)
-	{
-		UE_LOG(LogMotionWarping, Verbose, TEXT("MotionWarping: Marking RootMotionModifier for removal. Reason: Animation is not valid. Char: %s Current Montage: %s. Window: Animation: %s [%f %f] [%f %f]"),
-			*GetNameSafe(CharacterOwner), *GetNameSafe(Montage), *GetNameSafe(Animation.Get()), StartTime, EndTime, PreviousPosition, CurrentPosition);
-
-		SetState(ERootMotionModifierState::MarkedForRemoval);
-		return;
-	}
-
-	// Update playback times and weight
-	PreviousPosition = RootMotionMontageInstance->GetPreviousPosition();
-	CurrentPosition = RootMotionMontageInstance->GetPosition();
-	Weight = RootMotionMontageInstance->GetWeight();
-
-	// Mark for removal if the animation already passed the warping window
-	if (PreviousPosition >= EndTime)
-	{
-		UE_LOG(LogMotionWarping, Verbose, TEXT("MotionWarping: Marking RootMotionModifier for removal. Reason: Window has ended. Char: %s Animation: %s [%f %f] [%f %f]"),
-			*GetNameSafe(CharacterOwner), *GetNameSafe(Animation.Get()), StartTime, EndTime, PreviousPosition, CurrentPosition);
-
-		SetState(ERootMotionModifierState::MarkedForRemoval);
-		return;
-	}
-
-	// Check if we are inside the warping window
-	if (PreviousPosition >= StartTime && PreviousPosition < EndTime)
-	{
-		// If we were waiting, switch to active
-		if (GetState() == ERootMotionModifierState::Waiting)
-		{
-			SetState(ERootMotionModifierState::Active);
-		}
-	}
-
-	if (State == ERootMotionModifierState::Active)
-	{
-		if(UMotionWarpingComponent* OwnerComp = GetOwnerComponent())
-		{
-			OnUpdateDelegate.ExecuteIfBound(OwnerComp, Handle);
-		}
-	}
-}
-
-void FRootMotionModifier::SetState(ERootMotionModifierState NewState)
-{
-	if(State != NewState)
-	{
-		ERootMotionModifierState LastState = State;
-
-		State = NewState;
-
-		OnStateChanged(LastState);
-	}
-}
-
-void FRootMotionModifier::OnStateChanged(ERootMotionModifierState LastState)
-{
-	if (UMotionWarpingComponent* OwnerComp = GetOwnerComponent())
-	{
-		if (LastState != ERootMotionModifierState::Active && State == ERootMotionModifierState::Active)
-		{
-			OnActivateDelegate.ExecuteIfBound(OwnerComp, Handle);
-		}
-		else if (LastState == ERootMotionModifierState::Active && (State == ERootMotionModifierState::Disabled || State == ERootMotionModifierState::MarkedForRemoval))
-		{
-			OnDeactivateDelegate.ExecuteIfBound(OwnerComp, Handle);
-		}
-	}
-}
-
-// FRootMotionModifier_Warp
-///////////////////////////////////////////////////////////////
-
-void FRootMotionModifier_Warp::Update(UMotionWarpingComponent& OwnerComp_DEPRECATED)
-{
-	// Update playback times and state
-	FRootMotionModifier::Update(OwnerComp_DEPRECATED);
-
-	// Cache sync point transform and trigger OnSyncPointChanged if needed
-	const UMotionWarpingComponent* OwnerComp = GetOwnerComponent();
-	if (OwnerComp && GetState() == ERootMotionModifierState::Active)
-	{
-		const FMotionWarpingSyncPoint* SyncPointPtr = OwnerComp->FindSyncPoint(SyncPointName);
-		
-		// Disable if there is no sync point for us
-		if (SyncPointPtr == nullptr)
-		{
-			UE_LOG(LogMotionWarping, Verbose, TEXT("MotionWarping: Marking RootMotionModifier as Disabled. Reason: Invalid Sync Point (%s). Char: %s Animation: %s [%f %f] [%f %f]"),
-				*SyncPointName.ToString(), *GetNameSafe(OwnerComp->GetOwner()), *GetNameSafe(Animation.Get()), StartTime, EndTime, PreviousPosition, CurrentPosition);
-
-			SetState(ERootMotionModifierState::Disabled);
-			return;
-		}
-
-		// Get the warp point sent by the game
-		FTransform WarpPointTransformGame = FTransform(SyncPointPtr->GetRotation(), SyncPointPtr->GetLocation());
-
-		// Initialize our target transform (where the root should end at the end of the window) with the warp point sent by the game
-		FTransform TargetTransform = WarpPointTransformGame;
-
-		// Check if a warp point is defined in the animation. If so, we need to extract it and offset the target transform 
-		// the same amount the root bone is offset from the warp point in the animation
-		if(WarpPointAnimProvider != EWarpPointAnimProvider::None)
-		{
-			if(!CachedOffsetFromWarpPoint.IsSet())
-			{
-				if(const ACharacter* CharacterOwner = GetCharacterOwner())
-				{
-					// Inverse of mesh's relative rotation. Used to convert root and warp point in the animation from Y forward to X forward
-					const FTransform MeshCompRelativeRotInverse = FTransform(CharacterOwner->GetBaseRotationOffset().Inverse());
-
-					FTransform RootTransform = FTransform::Identity;
-					FTransform WarpPointTransform = FTransform::Identity;
-					if (WarpPointAnimProvider == EWarpPointAnimProvider::Static)
-					{
-						RootTransform = MeshCompRelativeRotInverse * UMotionWarpingUtilities::ExtractRootTransformFromAnimation(GetAnimation(), EndTime);
-						WarpPointTransform = MeshCompRelativeRotInverse * WarpPointAnimTransform;
-					}
-					else if (WarpPointAnimProvider == EWarpPointAnimProvider::Bone)
-					{
-						if (const USkeletalMeshComponent* Mesh = CharacterOwner->GetMesh())
-						{
-							if (const UAnimInstance* AnimInstance = Mesh->GetAnimInstance())
-							{
-								const FBoneContainer& FullBoneContainer = AnimInstance->GetRequiredBones();
-								const int32 BoneIndex = FullBoneContainer.GetPoseBoneIndexForBoneName(WarpPointAnimBoneName);
-								if (ensure(BoneIndex != INDEX_NONE))
-								{
-									TArray<FBoneIndexType> RequiredBoneIndexArray = { 0, (FBoneIndexType)BoneIndex };
-									FullBoneContainer.GetReferenceSkeleton().EnsureParentsExistAndSort(RequiredBoneIndexArray);
-
-									FBoneContainer LimitedBoneContainer(RequiredBoneIndexArray, FCurveEvaluationOption(false), *Animation->GetSkeleton());
-
-									FCSPose<FCompactPose> Pose;
-									UMotionWarpingUtilities::ExtractComponentSpacePose(GetAnimation(), LimitedBoneContainer, EndTime, false, Pose);
-
-									RootTransform = MeshCompRelativeRotInverse * Pose.GetComponentSpaceTransform(FCompactPoseBoneIndex(0));
-									WarpPointTransform = MeshCompRelativeRotInverse * Pose.GetComponentSpaceTransform(FCompactPoseBoneIndex(1));
-								}
-							}
-						}
-					}
-
-					// Calculate and cache the offset between the root and the warp point in the animation
-					CachedOffsetFromWarpPoint = WarpPointTransform.GetRelativeTransform(RootTransform);
-				}
-			}
-			
-			// Update Target Transform based on the offset between the root and the warp point in the animation
-			TargetTransform = CachedOffsetFromWarpPoint.GetValue() * WarpPointTransformGame;
-		}
-
-		if(!CachedTargetTransform.Equals(TargetTransform))
-		{
-			CachedTargetTransform = TargetTransform;
-
-			//DEPRECATED: Will be removed shortly
-			CachedSyncPoint = FMotionWarpingSyncPoint(CachedTargetTransform);
-
-			OnTargetTransformChanged();
-		}
-	}
-}
-
-FTransform FRootMotionModifier_Warp::ProcessRootMotion(UMotionWarpingComponent& OwnerComp_DEPRECATED, const FTransform& InRootMotion, float DeltaSeconds)
-{
-	const ACharacter* CharacterOwner = GetCharacterOwner();
-	if(CharacterOwner == nullptr)
-	{
-		return InRootMotion;
-	}
-
-	const FTransform& CharacterTransform = CharacterOwner->GetActorTransform();
-
-	FTransform FinalRootMotion = InRootMotion;
-
-	const FTransform RootMotionTotal = UMotionWarpingUtilities::ExtractRootMotionFromAnimation(Animation.Get(), PreviousPosition, EndTime);
-
-	if (bWarpTranslation)
-	{
-		FVector DeltaTranslation = InRootMotion.GetTranslation();
-
-		const FTransform RootMotionDelta = UMotionWarpingUtilities::ExtractRootMotionFromAnimation(Animation.Get(), PreviousPosition, FMath::Min(CurrentPosition, EndTime));
-
-		const float HorizontalDelta = RootMotionDelta.GetTranslation().Size2D();
-		const float HorizontalTarget = FVector::Dist2D(CharacterTransform.GetLocation(), GetTargetLocation());
-		const float HorizontalOriginal = RootMotionTotal.GetTranslation().Size2D();
-		const float HorizontalTranslationWarped = HorizontalOriginal != 0.f ? ((HorizontalDelta * HorizontalTarget) / HorizontalOriginal) : 0.f;
-
-		if (bInLocalSpace)
-		{
-			const FTransform MeshRelativeTransform = FTransform(CharacterOwner->GetBaseRotationOffset(), CharacterOwner->GetBaseTranslationOffset());
-			const FTransform MeshTransform = MeshRelativeTransform * CharacterOwner->GetActorTransform();
-			DeltaTranslation = MeshTransform.InverseTransformPositionNoScale(GetTargetLocation()).GetSafeNormal2D() * HorizontalTranslationWarped;
-		}
-		else
-		{
-			DeltaTranslation = (GetTargetLocation() - CharacterTransform.GetLocation()).GetSafeNormal2D() * HorizontalTranslationWarped;
-		}
-
-		if (!bIgnoreZAxis)
-		{
-			const FVector CapsuleBottomLocation = (CharacterOwner->GetActorLocation() - FVector::UpVector * CharacterOwner->GetSimpleCollisionHalfHeight());
-			const float VerticalDelta = RootMotionDelta.GetTranslation().Z;
-			const float VerticalTarget = GetTargetLocation().Z - CapsuleBottomLocation.Z;
-			const float VerticalOriginal = RootMotionTotal.GetTranslation().Z;
-			const float VerticalTranslationWarped = VerticalOriginal != 0.f ? ((VerticalDelta * VerticalTarget) / VerticalOriginal) : 0.f;
-
-			DeltaTranslation.Z = VerticalTranslationWarped;
-		}
-
-		FinalRootMotion.SetTranslation(DeltaTranslation);
-	}
-
-	if (bWarpRotation)
-	{
-		const FQuat WarpedRotation = WarpRotation(InRootMotion, RootMotionTotal, DeltaSeconds);
-		FinalRootMotion.SetRotation(WarpedRotation);
-	}
-
-	// Debug
-#if !(UE_BUILD_SHIPPING || UE_BUILD_TEST)
-	const int32 DebugLevel = FMotionWarpingCVars::CVarMotionWarpingDebug.GetValueOnGameThread();
-	if (DebugLevel > 0)
-	{
-		PrintLog(TEXT("FRootMotionModifier_Simple"), InRootMotion, FinalRootMotion);
-
-		if (DebugLevel >= 2)
-		{
-			const float DrawDebugDuration = FMotionWarpingCVars::CVarMotionWarpingDrawDebugDuration.GetValueOnGameThread();
-			DrawDebugCoordinateSystem(CharacterOwner->GetWorld(), GetTargetLocation(), GetTargetRotator(), 50.f, false, DrawDebugDuration, 0, 1.f);
-		}
-	}
-#endif
-
-	return FinalRootMotion;
-}
-
-FQuat FRootMotionModifier_Warp::GetTargetRotation() const
-{
-	if (RotationType == EMotionWarpRotationType::Default)
-	{
-		return CachedTargetTransform.GetRotation();
-	}
-	else if (RotationType == EMotionWarpRotationType::Facing)
-	{
-		if(const ACharacter* CharacterOwner = GetCharacterOwner())
-		{
-			const FTransform& CharacterTransform = CharacterOwner->GetActorTransform();
-			const FVector ToSyncPoint = (CachedTargetTransform.GetLocation() - CharacterTransform.GetLocation()).GetSafeNormal2D();
-			return FRotationMatrix::MakeFromXZ(ToSyncPoint, FVector::UpVector).ToQuat();
-		}
-	}
-
-	return FQuat::Identity;
-}
-
-FQuat FRootMotionModifier_Warp::WarpRotation(const FTransform& RootMotionDelta, const FTransform& RootMotionTotal, float DeltaSeconds)
-{
-	const ACharacter* CharacterOwner = GetCharacterOwner();
-	if (CharacterOwner == nullptr)
-	{
-		return FQuat::Identity;
-	}
-
-	const FTransform& CharacterTransform = CharacterOwner->GetActorTransform();
-	const FQuat CurrentRotation = CharacterTransform.GetRotation();
-	const FQuat TargetRotation = GetTargetRotation();
-	const float TimeRemaining = (EndTime - PreviousPosition) * WarpRotationTimeMultiplier;
-	const FQuat RemainingRootRotationInWorld = RootMotionTotal.GetRotation();
-	const FQuat CurrentPlusRemainingRootMotion = RemainingRootRotationInWorld * CurrentRotation;
-	const float PercentThisStep = FMath::Clamp(DeltaSeconds / TimeRemaining, 0.f, 1.f);
-	const FQuat TargetRotThisFrame = FQuat::Slerp(CurrentPlusRemainingRootMotion, TargetRotation, PercentThisStep);
-	const FQuat DeltaOut = TargetRotThisFrame * CurrentPlusRemainingRootMotion.Inverse();
-
-	return (DeltaOut * RootMotionDelta.GetRotation());
-}
-
-#if !(UE_BUILD_SHIPPING || UE_BUILD_TEST)
-void FRootMotionModifier_Warp::PrintLog(const FString& Name, const FTransform& OriginalRootMotion, const FTransform& WarpedRootMotion) const
-{
-	if(const ACharacter* CharacterOwner = GetCharacterOwner())
-	{
-		const FVector CurrentLocation = (CharacterOwner->GetActorLocation() - FVector::UpVector * CharacterOwner->GetSimpleCollisionHalfHeight());
-		const FVector CurrentToTarget = (GetTargetLocation() - CurrentLocation).GetSafeNormal2D();
-		const FVector FutureLocation = CurrentLocation + (bInLocalSpace ? (CharacterOwner->GetMesh()->ConvertLocalRootMotionToWorld(WarpedRootMotion)).GetTranslation() : WarpedRootMotion.GetTranslation());
-		const FRotator CurrentRotation = CharacterOwner->GetActorRotation();
-		const FRotator FutureRotation = (WarpedRootMotion.GetRotation() * CharacterOwner->GetActorQuat()).Rotator();
-		const float Dot = FVector::DotProduct(CharacterOwner->GetActorForwardVector(), CurrentToTarget);
-		const float CurrentDist2D = FVector::Dist2D(GetTargetLocation(), CurrentLocation);
-		const float FutureDist2D = FVector::Dist2D(GetTargetLocation(), FutureLocation);
-		const float DeltaSeconds = CharacterOwner->GetWorld()->GetDeltaSeconds();
-		const float Speed = WarpedRootMotion.GetTranslation().Size() / DeltaSeconds;
-		const float EndTimeOffset = CurrentPosition - EndTime;
-
-		UE_LOG(LogMotionWarping, Log, TEXT("MotionWarping: %s. NetMode: %d Char: %s Anim: %s Win: [%f %f][%f %f] DT: %f WT: %f EndTimeOffset: %f Dist2D: %f FutureDist2D: %f Dot: %f OriginalMotionDelta: %s (%f) FinalMotionDelta: %s (%f) Speed: %f Loc: %s FutureLoc: %s Rot: %s FutureRot: %s"),
-			*Name, (int32)CharacterOwner->GetWorld()->GetNetMode(), *GetNameSafe(CharacterOwner), *GetNameSafe(Animation.Get()), StartTime, EndTime, PreviousPosition, CurrentPosition, DeltaSeconds, CharacterOwner->GetWorld()->GetTimeSeconds(), EndTimeOffset, CurrentDist2D, FutureDist2D, Dot,
-			*OriginalRootMotion.GetTranslation().ToString(), OriginalRootMotion.GetTranslation().Size(), *WarpedRootMotion.GetTranslation().ToString(), WarpedRootMotion.GetTranslation().Size(), Speed,
-			*CurrentLocation.ToString(), *FutureLocation.ToString(), *CurrentRotation.ToCompactString(), *FutureRotation.ToCompactString());
-	}
-}
-#endif
-
-// URootMotionModifierConfig_Warp
-///////////////////////////////////////////////////////////////
-
-FRootMotionModifierHandle URootMotionModifierConfig_Warp::AddRootMotionModifierSimpleWarp(UMotionWarpingComponent* InMotionWarpingComp, const UAnimSequenceBase* InAnimation, float InStartTime, float InEndTime,
-	FName InSyncPointName, EWarpPointAnimProvider InWarpPointAnimProvider, FTransform InWarpPointAnimTransform, FName InWarpPointAnimBoneName,
-	bool bInWarpTranslation, bool bInIgnoreZAxis, bool bInWarpRotation, EMotionWarpRotationType InRotationType, float InWarpRotationTimeMultiplier)
-{
-	if (ensureAlways(InMotionWarpingComp))
-	{
-		UMotionModifier_Warp* NewModifier = NewObject<UMotionModifier_Warp>(InMotionWarpingComp);
-		NewModifier->Animation = InAnimation;
-		NewModifier->StartTime = InStartTime;
-		NewModifier->EndTime = InEndTime;
-		NewModifier->WarpTargetName = InSyncPointName;
-		NewModifier->WarpPointAnimProvider = InWarpPointAnimProvider;
-		NewModifier->WarpPointAnimTransform = InWarpPointAnimTransform;
-		NewModifier->WarpPointAnimBoneName = InWarpPointAnimBoneName;
-		NewModifier->bWarpTranslation = bInWarpTranslation;
-		NewModifier->bIgnoreZAxis = bInIgnoreZAxis;
-		NewModifier->bWarpRotation = bInWarpRotation;
-		NewModifier->RotationType = InRotationType;
-		NewModifier->WarpRotationTimeMultiplier = InWarpRotationTimeMultiplier;
-		return InMotionWarpingComp->AddModifier(NewModifier);
-	}
-
-	return FRootMotionModifierHandle::InvalidHandle;
-}
-
-// URootMotionModifierConfig_Scale
-///////////////////////////////////////////////////////////////
-
-FRootMotionModifierHandle URootMotionModifierConfig_Scale::AddRootMotionModifierScale(UMotionWarpingComponent* InMotionWarpingComp, const UAnimSequenceBase* InAnimation, float InStartTime, float InEndTime, FVector InScale)
-{
-	if (ensureAlways(InMotionWarpingComp))
-	{
-		TSharedPtr<FRootMotionModifier_Scale> NewModifier = MakeShared<FRootMotionModifier_Scale>();
-		NewModifier->Animation = InAnimation;
-		NewModifier->StartTime = InStartTime;
-		NewModifier->EndTime = InEndTime;
-		NewModifier->Scale = InScale;
-		return InMotionWarpingComp->AddRootMotionModifier(NewModifier);
-	}
-
-	return FRootMotionModifierHandle::InvalidHandle;
-}
-
-///////////////////////////////////////////////////////////////
-
-UMotionModifier::UMotionModifier(const FObjectInitializer& ObjectInitializer)
+URootMotionModifier::URootMotionModifier(const FObjectInitializer& ObjectInitializer)
 	: Super(ObjectInitializer)
 {
 }
 
-void UMotionModifier::Initialize(UMotionWarpingComponent* OwnerComp)
+void URootMotionModifier::Initialize(UMotionWarpingComponent* OwnerComp)
 {
 	check(OwnerComp);
 	OwnerComponent = OwnerComp;
 	Handle.GenerateNewHandle();
 }
 
-ACharacter* UMotionModifier::GetCharacterOwner() const
+ACharacter* URootMotionModifier::GetCharacterOwner() const
 {
 	UMotionWarpingComponent* OwnerComp = GetOwnerComponent();
 	return OwnerComp ? OwnerComp->GetCharacterOwner() : nullptr;
 }
 
-void UMotionModifier::Update()
+void URootMotionModifier::Update()
 {
 	const ACharacter* CharacterOwner = GetCharacterOwner();
 	if (CharacterOwner == nullptr)
@@ -471,7 +92,7 @@ void UMotionModifier::Update()
 	}
 }
 
-void UMotionModifier::SetState(ERootMotionModifierState NewState)
+void URootMotionModifier::SetState(ERootMotionModifierState NewState)
 {
 	if (State != NewState)
 	{
@@ -483,7 +104,7 @@ void UMotionModifier::SetState(ERootMotionModifierState NewState)
 	}
 }
 
-void UMotionModifier::OnStateChanged(ERootMotionModifierState LastState)
+void URootMotionModifier::OnStateChanged(ERootMotionModifierState LastState)
 {
 	if (UMotionWarpingComponent* OwnerComp = GetOwnerComponent())
 	{
@@ -498,12 +119,12 @@ void UMotionModifier::OnStateChanged(ERootMotionModifierState LastState)
 	}
 }
 
-UMotionModifier_Warp::UMotionModifier_Warp(const FObjectInitializer& ObjectInitializer)
+URootMotionModifier_Warp::URootMotionModifier_Warp(const FObjectInitializer& ObjectInitializer)
 	: Super(ObjectInitializer)
 {
 }
 
-void UMotionModifier_Warp::Update()
+void URootMotionModifier_Warp::Update()
 {
 	// Update playback times and state
 	Super::Update();
@@ -591,7 +212,7 @@ void UMotionModifier_Warp::Update()
 	}
 }
 
-FTransform UMotionModifier_Warp::ProcessRootMotion(const FTransform& InRootMotion, float DeltaSeconds)
+FTransform URootMotionModifier_Warp::ProcessRootMotion(const FTransform& InRootMotion, float DeltaSeconds)
 {
 	const ACharacter* CharacterOwner = GetCharacterOwner();
 	if (CharacterOwner == nullptr)
@@ -665,7 +286,7 @@ FTransform UMotionModifier_Warp::ProcessRootMotion(const FTransform& InRootMotio
 	return FinalRootMotion;
 }
 
-FQuat UMotionModifier_Warp::GetTargetRotation() const
+FQuat URootMotionModifier_Warp::GetTargetRotation() const
 {
 	if (RotationType == EMotionWarpRotationType::Default)
 	{
@@ -684,7 +305,7 @@ FQuat UMotionModifier_Warp::GetTargetRotation() const
 	return FQuat::Identity;
 }
 
-FQuat UMotionModifier_Warp::WarpRotation(const FTransform& RootMotionDelta, const FTransform& RootMotionTotal, float DeltaSeconds)
+FQuat URootMotionModifier_Warp::WarpRotation(const FTransform& RootMotionDelta, const FTransform& RootMotionTotal, float DeltaSeconds)
 {
 	const ACharacter* CharacterOwner = GetCharacterOwner();
 	if (CharacterOwner == nullptr)
@@ -706,7 +327,7 @@ FQuat UMotionModifier_Warp::WarpRotation(const FTransform& RootMotionDelta, cons
 }
 
 #if !(UE_BUILD_SHIPPING || UE_BUILD_TEST)
-void UMotionModifier_Warp::PrintLog(const FString& Name, const FTransform& OriginalRootMotion, const FTransform& WarpedRootMotion) const
+void URootMotionModifier_Warp::PrintLog(const FString& Name, const FTransform& OriginalRootMotion, const FTransform& WarpedRootMotion) const
 {
 	if (const ACharacter* CharacterOwner = GetCharacterOwner())
 	{
@@ -728,4 +349,32 @@ void UMotionModifier_Warp::PrintLog(const FString& Name, const FTransform& Origi
 			*CurrentLocation.ToString(), *FutureLocation.ToString(), *CurrentRotation.ToCompactString(), *FutureRotation.ToCompactString());
 	}
 }
-#endif
\ No newline at end of file
+#endif
+
+URootMotionModifier_Warp* URootMotionModifier_Warp::AddRootMotionModifierSimpleWarp(UMotionWarpingComponent* InMotionWarpingComp, const UAnimSequenceBase* InAnimation, float InStartTime, float InEndTime,
+	FName InWarpTargetName, EWarpPointAnimProvider InWarpPointAnimProvider, FTransform InWarpPointAnimTransform, FName InWarpPointAnimBoneName,
+	bool bInWarpTranslation, bool bInIgnoreZAxis, bool bInWarpRotation, EMotionWarpRotationType InRotationType, float InWarpRotationTimeMultiplier)
+{
+	if (ensureAlways(InMotionWarpingComp))
+	{
+		URootMotionModifier_Warp* NewModifier = NewObject<URootMotionModifier_Warp>(InMotionWarpingComp);
+		NewModifier->Animation = InAnimation;
+		NewModifier->StartTime = InStartTime;
+		NewModifier->EndTime = InEndTime;
+		NewModifier->WarpTargetName = InWarpTargetName;
+		NewModifier->WarpPointAnimProvider = InWarpPointAnimProvider;
+		NewModifier->WarpPointAnimTransform = InWarpPointAnimTransform;
+		NewModifier->WarpPointAnimBoneName = InWarpPointAnimBoneName;
+		NewModifier->bWarpTranslation = bInWarpTranslation;
+		NewModifier->bIgnoreZAxis = bInIgnoreZAxis;
+		NewModifier->bWarpRotation = bInWarpRotation;
+		NewModifier->RotationType = InRotationType;
+		NewModifier->WarpRotationTimeMultiplier = InWarpRotationTimeMultiplier;
+		
+		InMotionWarpingComp->AddModifier(NewModifier);
+
+		return NewModifier;
+	}
+
+	return nullptr;
+}
\ No newline at end of file
diff --git a/Engine/Plugins/Experimental/Animation/MotionWarping/Source/MotionWarping/Private/RootMotionModifier_AdjustmentBlendWarp.cpp b/Engine/Plugins/Experimental/Animation/MotionWarping/Source/MotionWarping/Private/RootMotionModifier_AdjustmentBlendWarp.cpp
--- a/Engine/Plugins/Experimental/Animation/MotionWarping/Source/MotionWarping/Private/RootMotionModifier_AdjustmentBlendWarp.cpp
+++ b/Engine/Plugins/Experimental/Animation/MotionWarping/Source/MotionWarping/Private/RootMotionModifier_AdjustmentBlendWarp.cpp
@@ -14,18 +14,55 @@
 DECLARE_CYCLE_STAT(TEXT("MotionWarping PrecomputeWarpedTracks"), STAT_MotionWarping_PrecomputeWarpedTracks, STATGROUP_Anim);
 DECLARE_CYCLE_STAT(TEXT("MotionWarping ExtractMotionDelta"), STAT_MotionWarping_ExtractMotionDelta, STATGROUP_Anim);
 
-FRootMotionModifier_AdjustmentBlendWarp::FRootMotionModifier_AdjustmentBlendWarp()
+void URootMotionModifierConfig_AdjustmentBlendWarp::GetIKBoneTransformAndAlpha(ACharacter* Character, FName BoneName, FTransform& OutTransform, float& OutAlpha)
 {
-	bInLocalSpace = true; 
+	URootMotionModifier_AdjustmentBlendWarp::GetAdjustmentBlendIKBoneTransformAndAlpha(Character, BoneName, OutTransform, OutAlpha);
 }
 
-void FRootMotionModifier_AdjustmentBlendWarp::ExtractBoneTransformAtTime(FTransform& OutTransform, const FName& BoneName, float Time) const
+// URootMotionModifierConfig_AdjustmentBlendWarp
+///////////////////////////////////////////////////////////////
+
+FRootMotionModifierHandle URootMotionModifierConfig_AdjustmentBlendWarp::AddRootMotionModifierAdjustmentBlendWarp(UMotionWarpingComponent* InMotionWarpingComp, const UAnimSequenceBase* InAnimation, float InStartTime, float InEndTime, 
+	FName InSyncPointName, EWarpPointAnimProvider InWarpPointAnimProvider, FTransform InWarpPointAnimTransform, FName InWarpPointAnimBoneName,
+	bool bInWarpTranslation, bool bInIgnoreZAxis, bool bInWarpRotation, bool bInWarpIKBones, const TArray<FName>& InIKBones)
+{
+	if (ensureAlways(InMotionWarpingComp))
+	{
+		URootMotionModifier_AdjustmentBlendWarp* NewModifier = NewObject<URootMotionModifier_AdjustmentBlendWarp>(InMotionWarpingComp);
+		NewModifier->Animation = InAnimation;
+		NewModifier->StartTime = InStartTime;
+		NewModifier->EndTime = InEndTime;
+		NewModifier->WarpTargetName = InSyncPointName;
+		NewModifier->WarpPointAnimProvider = InWarpPointAnimProvider;
+		NewModifier->WarpPointAnimTransform = InWarpPointAnimTransform;
+		NewModifier->WarpPointAnimBoneName = InWarpPointAnimBoneName;
+		NewModifier->bWarpTranslation = bInWarpTranslation;
+		NewModifier->bIgnoreZAxis = bInIgnoreZAxis;
+		NewModifier->bWarpRotation = bInWarpRotation;
+		NewModifier->bWarpIKBones = bInWarpIKBones;
+		NewModifier->IKBones = InIKBones;
+
+		return InMotionWarpingComp->AddModifier(NewModifier);
+	}
+
+	return FRootMotionModifierHandle::InvalidHandle;
+}
+
+/////////////////////////////////////////////////////////
+
+URootMotionModifier_AdjustmentBlendWarp::URootMotionModifier_AdjustmentBlendWarp(const FObjectInitializer& ObjectInitializer)
+	: Super(ObjectInitializer)
+{
+	bInLocalSpace = true;
+}
+
+void URootMotionModifier_AdjustmentBlendWarp::ExtractBoneTransformAtTime(FTransform& OutTransform, const FName& BoneName, float Time) const
 {
 	const int32 TrackIndex = Result.TrackNames.IndexOfByKey(BoneName);
 	ExtractBoneTransformAtTime(OutTransform, TrackIndex, Time);
 }
 
-void FRootMotionModifier_AdjustmentBlendWarp::ExtractBoneTransformAtTime(FTransform& OutTransform, int32 TrackIndex, float Time) const
+void URootMotionModifier_AdjustmentBlendWarp::ExtractBoneTransformAtTime(FTransform& OutTransform, int32 TrackIndex, float Time) const
 {
 	if (!Result.AnimationTracks.IsValidIndex(TrackIndex))
 	{
@@ -40,7 +77,7 @@ void FRootMotionModifier_AdjustmentBlendWarp::ExtractBoneTransformAtTime(FTransf
 	FAnimationUtils::ExtractTransformFromTrack(RemappedTime, TotalFrames, TrackLength, Result.AnimationTracks[TrackIndex], EAnimInterpolationType::Linear, OutTransform);
 }
 
-void FRootMotionModifier_AdjustmentBlendWarp::ExtractBoneTransformAtFrame(FTransform& OutTransform, int32 TrackIndex, int32 Frame) const
+void URootMotionModifier_AdjustmentBlendWarp::ExtractBoneTransformAtFrame(FTransform& OutTransform, int32 TrackIndex, int32 Frame) const
 {
 	if (!Result.AnimationTracks.IsValidIndex(TrackIndex) || !Result.AnimationTracks[TrackIndex].PosKeys.IsValidIndex(Frame))
 	{
@@ -51,12 +88,12 @@ void FRootMotionModifier_AdjustmentBlendWarp::ExtractBoneTransformAtFrame(FTrans
 	FAnimationUtils::ExtractTransformForFrameFromTrack(Result.AnimationTracks[TrackIndex], Frame, OutTransform);
 }
 
-void FRootMotionModifier_AdjustmentBlendWarp::OnTargetTransformChanged()
+void URootMotionModifier_AdjustmentBlendWarp::OnTargetTransformChanged()
 {
 	const ACharacter* CharacterOwner = GetCharacterOwner();
 	const USkeletalMeshComponent* SkelMeshComp = CharacterOwner ? CharacterOwner->GetMesh() : nullptr;
 
-	if(SkelMeshComp)
+	if (SkelMeshComp)
 	{
 		ActualStartTime = PreviousPosition;
 		CachedRootMotion = UMotionWarpingUtilities::ExtractRootMotionFromAnimation(Animation.Get(), ActualStartTime, EndTime);
@@ -68,7 +105,7 @@ void FRootMotionModifier_AdjustmentBlendWarp::OnTargetTransformChanged()
 	}
 }
 
-FTransform FRootMotionModifier_AdjustmentBlendWarp::ProcessRootMotion(UMotionWarpingComponent& OwnerComp_DEPRECATED, const FTransform& InRootMotion, float DeltaSeconds)
+FTransform URootMotionModifier_AdjustmentBlendWarp::ProcessRootMotion(const FTransform& InRootMotion, float DeltaSeconds)
 {
 	// If warped tracks has not been generated yet, do it now
 	if (Result.GetNum() == 0)
@@ -85,7 +122,7 @@ FTransform FRootMotionModifier_AdjustmentBlendWarp::ProcessRootMotion(UMotionWar
 	{
 		PrintLog(TEXT("FRootMotionModifier_AdjustmentBlend"), InRootMotion, FinalRootMotion);
 
-		if(FMotionWarpingCVars::CVarMotionWarpingDebug.GetValueOnGameThread() >= 2)
+		if (FMotionWarpingCVars::CVarMotionWarpingDebug.GetValueOnGameThread() >= 2)
 		{
 			const float DrawDebugDuration = FMotionWarpingCVars::CVarMotionWarpingDrawDebugDuration.GetValueOnGameThread();
 			DrawDebugWarpedTracks(DrawDebugDuration);
@@ -96,14 +133,14 @@ FTransform FRootMotionModifier_AdjustmentBlendWarp::ProcessRootMotion(UMotionWar
 	return FinalRootMotion;
 }
 
-void FRootMotionModifier_AdjustmentBlendWarp::PrecomputeWarpedTracks()
+void URootMotionModifier_AdjustmentBlendWarp::PrecomputeWarpedTracks()
 {
 	SCOPE_CYCLE_COUNTER(STAT_MotionWarping_PrecomputeWarpedTracks);
 
 	// First, extract pose at the end of the window for the bones we are going to warp
 
 	const ACharacter* CharacterOwner = GetCharacterOwner();
-	if(CharacterOwner == nullptr)
+	if (CharacterOwner == nullptr)
 	{
 		return;
 	}
@@ -120,7 +157,7 @@ void FRootMotionModifier_AdjustmentBlendWarp::PrecomputeWarpedTracks()
 		for (const FName& BoneName : IKBones)
 		{
 			const int32 BoneIndex = BoneContainer.GetPoseBoneIndexForBoneName(BoneName);
-			if(BoneIndex != INDEX_NONE)
+			if (BoneIndex != INDEX_NONE)
 			{
 				RequiredBoneIndexArray.Add(BoneIndex);
 			}
@@ -193,12 +230,12 @@ void FRootMotionModifier_AdjustmentBlendWarp::PrecomputeWarpedTracks()
 	//@todo_fer: We could extract and cache this offline when the WarpingWindow is created
 	const float SampleRate = 1 / 60.f;
 	FMotionDeltaTrackContainer MotionDeltaTracks;
-	FRootMotionModifier_AdjustmentBlendWarp::ExtractMotionDeltaFromRange(RequiredBones, Animation.Get(), ActualStartTime, EndTime, SampleRate, MotionDeltaTracks);
+	URootMotionModifier_AdjustmentBlendWarp::ExtractMotionDeltaFromRange(RequiredBones, Animation.Get(), ActualStartTime, EndTime, SampleRate, MotionDeltaTracks);
 
-	FRootMotionModifier_AdjustmentBlendWarp::AdjustmentBlendWarp(RequiredBones, AdditivePose, MotionDeltaTracks, Result);
+	URootMotionModifier_AdjustmentBlendWarp::AdjustmentBlendWarp(RequiredBones, AdditivePose, MotionDeltaTracks, Result);
 }
 
-void FRootMotionModifier_AdjustmentBlendWarp::ExtractMotionDeltaFromRange(const FBoneContainer& BoneContainer, const UAnimSequenceBase* Animation, float StartTime, float EndTime, float SampleRate, FMotionDeltaTrackContainer& OutMotionDeltaTracks)
+void URootMotionModifier_AdjustmentBlendWarp::ExtractMotionDeltaFromRange(const FBoneContainer& BoneContainer, const UAnimSequenceBase* Animation, float StartTime, float EndTime, float SampleRate, FMotionDeltaTrackContainer& OutMotionDeltaTracks)
 {
 	SCOPE_CYCLE_COUNTER(STAT_MotionWarping_ExtractMotionDelta);
 
@@ -265,7 +302,7 @@ void FRootMotionModifier_AdjustmentBlendWarp::ExtractMotionDeltaFromRange(const
 	}
 }
 
-void FRootMotionModifier_AdjustmentBlendWarp::AdjustmentBlendWarp(const FBoneContainer& BoneContainer, const FCSPose<FCompactPose>& AdditivePose, const FMotionDeltaTrackContainer& MotionDeltaTracks, FAnimSequenceTrackContainer& Output)
+void URootMotionModifier_AdjustmentBlendWarp::AdjustmentBlendWarp(const FBoneContainer& BoneContainer, const FCSPose<FCompactPose>& AdditivePose, const FMotionDeltaTrackContainer& MotionDeltaTracks, FAnimSequenceTrackContainer& Output)
 {
 	auto CalculateAdditive = [](const FVector& Total, const FVector& Delta, const FVector& Additive, const FVector& PreviousAdditive, float Alpha)
 	{
@@ -337,7 +374,7 @@ void FRootMotionModifier_AdjustmentBlendWarp::AdjustmentBlendWarp(const FBoneCon
 	}
 }
 
-FTransform FRootMotionModifier_AdjustmentBlendWarp::ExtractWarpedRootMotion() const
+FTransform URootMotionModifier_AdjustmentBlendWarp::ExtractWarpedRootMotion() const
 {
 	FTransform StartRootTransform;
 	ExtractBoneTransformAtTime(StartRootTransform, 0, PreviousPosition);
@@ -349,7 +386,7 @@ FTransform FRootMotionModifier_AdjustmentBlendWarp::ExtractWarpedRootMotion() co
 }
 
 #if !(UE_BUILD_SHIPPING || UE_BUILD_TEST)
-void FRootMotionModifier_AdjustmentBlendWarp::DrawDebugWarpedTracks(float DrawDuration) const
+void URootMotionModifier_AdjustmentBlendWarp::DrawDebugWarpedTracks(float DrawDuration) const
 {
 	const ACharacter* CharacterOwner = GetCharacterOwner();
 	const UWorld* World = CharacterOwner ? CharacterOwner->GetWorld() : nullptr;
@@ -387,7 +424,7 @@ void FRootMotionModifier_AdjustmentBlendWarp::DrawDebugWarpedTracks(float DrawDu
 }
 #endif
 
-void FRootMotionModifier_AdjustmentBlendWarp::GetIKBoneTransformAndAlpha(FName BoneName, FTransform& OutTransform, float& OutAlpha) const
+void URootMotionModifier_AdjustmentBlendWarp::GetIKBoneTransformAndAlpha(FName BoneName, FTransform& OutTransform, float& OutAlpha) const
 {
 	if (Result.GetNum() == 0 || !bWarpIKBones || !IKBones.Contains(BoneName))
 	{
@@ -406,7 +443,7 @@ void FRootMotionModifier_AdjustmentBlendWarp::GetIKBoneTransformAndAlpha(FName B
 	OutAlpha = Weight;
 }
 
-void URootMotionModifierConfig_AdjustmentBlendWarp::GetIKBoneTransformAndAlpha(ACharacter* Character, FName BoneName, FTransform& OutTransform, float& OutAlpha)
+void URootMotionModifier_AdjustmentBlendWarp::GetAdjustmentBlendIKBoneTransformAndAlpha(ACharacter* Character, FName BoneName, FTransform& OutTransform, float& OutAlpha)
 {
 	OutTransform = FTransform::Identity;
 	OutAlpha = 0.f;
@@ -414,27 +451,11 @@ void URootMotionModifierConfig_AdjustmentBlendWarp::GetIKBoneTransformAndAlpha(A
 	const UMotionWarpingComponent* MotionWarpingComp = Character ? Character->FindComponentByClass<UMotionWarpingComponent>() : nullptr;
 	if (MotionWarpingComp)
 	{
-		for (const auto& Modifier : MotionWarpingComp->GetRootMotionModifiers())
-		{
-			if (Modifier.IsValid() && Modifier->GetState() == ERootMotionModifierState::Active && Modifier->GetScriptStruct()->IsChildOf(FRootMotionModifier_AdjustmentBlendWarp::StaticStruct()))
-			{
-				// We must check if the Animation for the modifier is still relevant because in RootMotionFromMontageOnlyMode the montage could be aborted 
-				// but the modifier will remain in the list until the next update
-				const FAnimMontageInstance* RootMotionMontageInstance = Character->GetRootMotionAnimMontageInstance();
-				const UAnimMontage* Montage = RootMotionMontageInstance ? ToRawPtr(RootMotionMontageInstance->Montage) : nullptr;
-				if (Modifier->Animation == Montage)
-				{
-					static_cast<const FRootMotionModifier_AdjustmentBlendWarp*>(Modifier.Get())->GetIKBoneTransformAndAlpha(BoneName, OutTransform, OutAlpha);
-					return;
-				}
-			}
-		}
-
-		for (UMotionModifier* Modifier : MotionWarpingComp->GetModifiers())
+		for (URootMotionModifier* Modifier : MotionWarpingComp->GetModifiers())
 		{
 			if (Modifier && Modifier->GetState() == ERootMotionModifierState::Active)
 			{
-				if(UMotionModifier_AdjustmentBlendWarp* AdjustmentBlendWarpModifier = Cast<UMotionModifier_AdjustmentBlendWarp>(Modifier))
+				if (URootMotionModifier_AdjustmentBlendWarp* AdjustmentBlendWarpModifier = Cast<URootMotionModifier_AdjustmentBlendWarp>(Modifier))
 				{
 					// We must check if the Animation for the modifier is still relevant because in RootMotionFromMontageOnlyMode the montage could be aborted 
 					// but the modifier will remain in the list until the next update
@@ -451,20 +472,17 @@ void URootMotionModifierConfig_AdjustmentBlendWarp::GetIKBoneTransformAndAlpha(A
 	}
 }
 
-// URootMotionModifierConfig_AdjustmentBlendWarp
-///////////////////////////////////////////////////////////////
-
-FRootMotionModifierHandle URootMotionModifierConfig_AdjustmentBlendWarp::AddRootMotionModifierAdjustmentBlendWarp(UMotionWarpingComponent* InMotionWarpingComp, const UAnimSequenceBase* InAnimation, float InStartTime, float InEndTime, 
-	FName InSyncPointName, EWarpPointAnimProvider InWarpPointAnimProvider, FTransform InWarpPointAnimTransform, FName InWarpPointAnimBoneName,
+URootMotionModifier_AdjustmentBlendWarp* URootMotionModifier_AdjustmentBlendWarp::AddRootMotionModifierAdjustmentBlendWarp(UMotionWarpingComponent* InMotionWarpingComp, const UAnimSequenceBase* InAnimation, float InStartTime, float InEndTime,
+	FName InWarpTargetName, EWarpPointAnimProvider InWarpPointAnimProvider, FTransform InWarpPointAnimTransform, FName InWarpPointAnimBoneName,
 	bool bInWarpTranslation, bool bInIgnoreZAxis, bool bInWarpRotation, bool bInWarpIKBones, const TArray<FName>& InIKBones)
 {
 	if (ensureAlways(InMotionWarpingComp))
 	{
-		UMotionModifier_AdjustmentBlendWarp* NewModifier = NewObject<UMotionModifier_AdjustmentBlendWarp>(InMotionWarpingComp);
+		URootMotionModifier_AdjustmentBlendWarp* NewModifier = NewObject<URootMotionModifier_AdjustmentBlendWarp>(InMotionWarpingComp);
 		NewModifier->Animation = InAnimation;
 		NewModifier->StartTime = InStartTime;
 		NewModifier->EndTime = InEndTime;
-		NewModifier->WarpTargetName = InSyncPointName;
+		NewModifier->WarpTargetName = InWarpTargetName;
 		NewModifier->WarpPointAnimProvider = InWarpPointAnimProvider;
 		NewModifier->WarpPointAnimTransform = InWarpPointAnimTransform;
 		NewModifier->WarpPointAnimBoneName = InWarpPointAnimBoneName;
@@ -474,403 +492,10 @@ FRootMotionModifierHandle URootMotionModifierConfig_AdjustmentBlendWarp::AddRoot
 		NewModifier->bWarpIKBones = bInWarpIKBones;
 		NewModifier->IKBones = InIKBones;
 
-		return InMotionWarpingComp->AddModifier(NewModifier);
-	}
-
-	return FRootMotionModifierHandle::InvalidHandle;
-}
-
-/////////////////////////////////////////////////////////
-
-UMotionModifier_AdjustmentBlendWarp::UMotionModifier_AdjustmentBlendWarp(const FObjectInitializer& ObjectInitializer)
-	: Super(ObjectInitializer)
-{
-	bInLocalSpace = true;
-}
-
-void UMotionModifier_AdjustmentBlendWarp::ExtractBoneTransformAtTime(FTransform& OutTransform, const FName& BoneName, float Time) const
-{
-	const int32 TrackIndex = Result.TrackNames.IndexOfByKey(BoneName);
-	ExtractBoneTransformAtTime(OutTransform, TrackIndex, Time);
-}
-
-void UMotionModifier_AdjustmentBlendWarp::ExtractBoneTransformAtTime(FTransform& OutTransform, int32 TrackIndex, float Time) const
-{
-	if (!Result.AnimationTracks.IsValidIndex(TrackIndex))
-	{
-		OutTransform = FTransform::Identity;
-		return;
-	}
-
-	const int32 TotalFrames = FMath::Max(Result.AnimationTracks[TrackIndex].PosKeys.Num(), Result.AnimationTracks[TrackIndex].RotKeys.Num());
-	const float TrackLength = (EndTime - ActualStartTime);
-	const float TimePercent = (Time - ActualStartTime) / (EndTime - ActualStartTime);
-	const float RemappedTime = TimePercent * TrackLength;
-	FAnimationUtils::ExtractTransformFromTrack(RemappedTime, TotalFrames, TrackLength, Result.AnimationTracks[TrackIndex], EAnimInterpolationType::Linear, OutTransform);
-}
-
-void UMotionModifier_AdjustmentBlendWarp::ExtractBoneTransformAtFrame(FTransform& OutTransform, int32 TrackIndex, int32 Frame) const
-{
-	if (!Result.AnimationTracks.IsValidIndex(TrackIndex) || !Result.AnimationTracks[TrackIndex].PosKeys.IsValidIndex(Frame))
-	{
-		OutTransform = FTransform::Identity;
-		return;
-	}
-
-	FAnimationUtils::ExtractTransformForFrameFromTrack(Result.AnimationTracks[TrackIndex], Frame, OutTransform);
-}
-
-void UMotionModifier_AdjustmentBlendWarp::OnTargetTransformChanged()
-{
-	const ACharacter* CharacterOwner = GetCharacterOwner();
-	const USkeletalMeshComponent* SkelMeshComp = CharacterOwner ? CharacterOwner->GetMesh() : nullptr;
-
-	if (SkelMeshComp)
-	{
-		ActualStartTime = PreviousPosition;
-		CachedRootMotion = UMotionWarpingUtilities::ExtractRootMotionFromAnimation(Animation.Get(), ActualStartTime, EndTime);
-		CachedMeshRelativeTransform = SkelMeshComp->GetRelativeTransform();
-		CachedMeshTransform = SkelMeshComp->GetComponentTransform();
-
-		Result.AnimationTracks.Reset();
-		Result.TrackNames.Reset();
-	}
-}
-
-FTransform UMotionModifier_AdjustmentBlendWarp::ProcessRootMotion(const FTransform& InRootMotion, float DeltaSeconds)
-{
-	// If warped tracks has not been generated yet, do it now
-	if (Result.GetNum() == 0)
-	{
-		PrecomputeWarpedTracks();
-	}
-
-	// Extract root motion from warped tracks
-	const FTransform FinalRootMotion = ExtractWarpedRootMotion();
-
-	// Debug
-#if !(UE_BUILD_SHIPPING || UE_BUILD_TEST)
-	if (FMotionWarpingCVars::CVarMotionWarpingDebug.GetValueOnGameThread() > 0)
-	{
-		PrintLog(TEXT("FRootMotionModifier_AdjustmentBlend"), InRootMotion, FinalRootMotion);
-
-		if (FMotionWarpingCVars::CVarMotionWarpingDebug.GetValueOnGameThread() >= 2)
-		{
-			const float DrawDebugDuration = FMotionWarpingCVars::CVarMotionWarpingDrawDebugDuration.GetValueOnGameThread();
-			DrawDebugWarpedTracks(DrawDebugDuration);
-		}
-	}
-#endif
-
-	return FinalRootMotion;
-}
-
-void UMotionModifier_AdjustmentBlendWarp::PrecomputeWarpedTracks()
-{
-	SCOPE_CYCLE_COUNTER(STAT_MotionWarping_PrecomputeWarpedTracks);
-
-	// First, extract pose at the end of the window for the bones we are going to warp
-
-	const ACharacter* CharacterOwner = GetCharacterOwner();
-	if (CharacterOwner == nullptr)
-	{
-		return;
-	}
-
-	const FBoneContainer& BoneContainer = CharacterOwner->GetMesh()->GetAnimInstance()->GetRequiredBones();
-
-	// Init FBoneContainer with only the bones that we are interested in
-	TArray<FBoneIndexType> RequiredBoneIndexArray;
-	RequiredBoneIndexArray.Add(0);
-
-	const bool bShouldWarpIKBones = bWarpIKBones && IKBones.Num() > 0;
-	if (bShouldWarpIKBones)
-	{
-		for (const FName& BoneName : IKBones)
-		{
-			const int32 BoneIndex = BoneContainer.GetPoseBoneIndexForBoneName(BoneName);
-			if (BoneIndex != INDEX_NONE)
-			{
-				RequiredBoneIndexArray.Add(BoneIndex);
-			}
-		}
-
-		BoneContainer.GetReferenceSkeleton().EnsureParentsExistAndSort(RequiredBoneIndexArray);
-	}
-
-	// Init BoneContainer
-	FBoneContainer RequiredBones(RequiredBoneIndexArray, FCurveEvaluationOption(false), *BoneContainer.GetAsset());
-
-	// Extract pose
-	FCSPose<FCompactPose> CSPose;
-	UMotionWarpingUtilities::ExtractComponentSpacePose(Animation.Get(), RequiredBones, EndTime, true, CSPose);
-
-	// Second, calculate additive pose
-
-	//Calculate additive translation for root bone
-	FVector RootTargetLocation = CachedMeshTransform.InverseTransformPositionNoScale(GetTargetLocation());
-
-	FVector RootTotalAdditiveTranslation = FVector::ZeroVector;
-	if (bWarpTranslation)
-	{
-		RootTotalAdditiveTranslation = RootTargetLocation - CachedRootMotion.GetLocation();
-
-		if (bIgnoreZAxis)
-		{
-			RootTotalAdditiveTranslation.Z = 0.f;
-		}
-	}
-
-	// Calculate additive rotation for root bone
-	FQuat RootTotalAdditiveRotation = FQuat::Identity;
-	if (bWarpRotation)
-	{
-		const FQuat TargetRotation = GetTargetRotation();
-		const FQuat OriginalRotation = CachedMeshRelativeTransform.GetRotation().Inverse() * (CachedRootMotion * CachedMeshTransform).GetRotation();
-		RootTotalAdditiveRotation = FQuat::FindBetweenNormals(OriginalRotation.GetForwardVector(), TargetRotation.GetForwardVector());
-	}
-
-	// Init Additive Pose
-	FCSPose<FCompactPose> AdditivePose;
-	AdditivePose.InitPose(&RequiredBones);
-	AdditivePose.SetComponentSpaceTransform(FCompactPoseBoneIndex(0), FTransform(RootTotalAdditiveRotation, RootTotalAdditiveTranslation));
-
-	// Calculate and add additive pose for IK bones
-	if (bShouldWarpIKBones)
-	{
-		const FTransform RootTargetPoseCS = FTransform((RootTotalAdditiveRotation * CachedRootMotion.GetRotation()), CachedRootMotion.GetTranslation() + RootTotalAdditiveTranslation);
-		for (int32 Idx = 1; Idx < CSPose.GetPose().GetNumBones(); Idx++)
-		{
-			const FName BoneName = RequiredBones.GetReferenceSkeleton().GetBoneName(RequiredBones.GetBoneIndicesArray()[Idx]);
-			if (IKBones.Contains(BoneName))
-			{
-				const int32 BoneIdx = Idx;
-				const FTransform BonePoseCS = CSPose.GetComponentSpaceTransform(FCompactPoseBoneIndex(BoneIdx));
-
-				const FTransform BoneTargetPoseCS = BonePoseCS * RootTargetPoseCS;
-				const FTransform BoneOriginalPoseCS = BonePoseCS * CachedRootMotion;
-
-				const FVector TotalAdditiveTranslation = BoneTargetPoseCS.GetLocation() - BoneOriginalPoseCS.GetLocation();
-
-				AdditivePose.SetComponentSpaceTransform(FCompactPoseBoneIndex(Idx), FTransform(TotalAdditiveTranslation));
-			}
-		}
-	}
-
-	// Finally, run adjustment blending to generate the warped poses for each bone
-
-	//@todo_fer: We could extract and cache this offline when the WarpingWindow is created
-	const float SampleRate = 1 / 60.f;
-	FMotionDeltaTrackContainer MotionDeltaTracks;
-	UMotionModifier_AdjustmentBlendWarp::ExtractMotionDeltaFromRange(RequiredBones, Animation.Get(), ActualStartTime, EndTime, SampleRate, MotionDeltaTracks);
-
-	UMotionModifier_AdjustmentBlendWarp::AdjustmentBlendWarp(RequiredBones, AdditivePose, MotionDeltaTracks, Result);
-}
-
-void UMotionModifier_AdjustmentBlendWarp::ExtractMotionDeltaFromRange(const FBoneContainer& BoneContainer, const UAnimSequenceBase* Animation, float StartTime, float EndTime, float SampleRate, FMotionDeltaTrackContainer& OutMotionDeltaTracks)
-{
-	SCOPE_CYCLE_COUNTER(STAT_MotionWarping_ExtractMotionDelta);
-
-	check(Animation);
-
-	const int32 TotalFrames = FMath::CeilToInt((EndTime - StartTime) / SampleRate) + 1;
-	const int32 TotalBones = BoneContainer.GetCompactPoseNumBones();
-
-	FCompactPose FirstFramePose;
-	UMotionWarpingUtilities::ExtractLocalSpacePose(Animation, BoneContainer, StartTime, false, FirstFramePose);
-	const FTransform RootTransformFirstFrame = FirstFramePose[FCompactPoseBoneIndex(0)];
-
-	FCSPose<FCompactPose> CSPose;
-	UMotionWarpingUtilities::ExtractComponentSpacePose(Animation, BoneContainer, StartTime, false, CSPose);
-
-	OutMotionDeltaTracks.Init(TotalBones);
-	for (int32 BoneIdx = 0; BoneIdx < TotalBones; BoneIdx++)
-	{
-		const FTransform& BoneTransform = (CSPose.GetComponentSpaceTransform(FCompactPoseBoneIndex(BoneIdx)).GetRelativeTransform(RootTransformFirstFrame));
-
-		FMotionDeltaTrack& NewTrack = OutMotionDeltaTracks.Tracks.AddDefaulted_GetRef();
-
-		NewTrack.BoneTransformTrack.Reserve(TotalFrames);
-		NewTrack.DeltaTranslationTrack.Reserve(TotalFrames);
-		NewTrack.DeltaRotationTrack.Reserve(TotalFrames);
-
-		NewTrack.BoneTransformTrack.Add(BoneTransform);
-		NewTrack.DeltaTranslationTrack.Add(FVector::ZeroVector);
-		NewTrack.DeltaRotationTrack.Add(FRotator::ZeroRotator);
-
-		NewTrack.TotalTranslation = FVector::ZeroVector;
-		NewTrack.TotalRotation = FRotator::ZeroRotator;
-	}
-
-	float Time = StartTime;
-	while (Time < EndTime)
-	{
-		Time += FMath::Min(SampleRate, EndTime - Time);
-
-		UMotionWarpingUtilities::ExtractComponentSpacePose(Animation, BoneContainer, Time, false, CSPose);
-
-		for (int32 BoneIdx = 0; BoneIdx < TotalBones; BoneIdx++)
-		{
-			FMotionDeltaTrack& Track = OutMotionDeltaTracks.Tracks[BoneIdx];
-
-			const FTransform& BoneTransform = (CSPose.GetComponentSpaceTransform(FCompactPoseBoneIndex(BoneIdx)).GetRelativeTransform(RootTransformFirstFrame));
-
-			const FTransform& LastBoneTransform = Track.BoneTransformTrack.Last();
-
-			Track.BoneTransformTrack.Add(BoneTransform);
-
-			const FVector Translation = BoneTransform.GetTranslation();
-			const FVector LastTranslation = LastBoneTransform.GetTranslation();
-			const FVector DeltaTranslation(FMath::Abs(Translation.X - LastTranslation.X), FMath::Abs(Translation.Y - LastTranslation.Y), FMath::Abs(Translation.Z - LastTranslation.Z));
-			Track.DeltaTranslationTrack.Add(DeltaTranslation);
-			Track.TotalTranslation = Track.TotalTranslation + DeltaTranslation;
-
-			const FRotator Rotation = BoneTransform.GetRotation().Rotator();
-			const FRotator LastRotation = LastBoneTransform.GetRotation().Rotator();
-			const FRotator DeltaRotation(FMath::Abs(Rotation.Pitch - LastRotation.Pitch), FMath::Abs(Rotation.Yaw - LastRotation.Yaw), FMath::Abs(Rotation.Roll - LastRotation.Roll));
-			Track.DeltaRotationTrack.Add(DeltaRotation);
-			Track.TotalRotation = Track.TotalRotation + DeltaRotation;
-		}
-	}
-}
-
-void UMotionModifier_AdjustmentBlendWarp::AdjustmentBlendWarp(const FBoneContainer& BoneContainer, const FCSPose<FCompactPose>& AdditivePose, const FMotionDeltaTrackContainer& MotionDeltaTracks, FAnimSequenceTrackContainer& Output)
-{
-	auto CalculateAdditive = [](const FVector& Total, const FVector& Delta, const FVector& Additive, const FVector& PreviousAdditive, float Alpha)
-	{
-		FVector CurrentAdditive = FVector::ZeroVector;
-		for (int32 Idx = 0; Idx < 3; Idx++)
-		{
-			if (!FMath::IsNearlyZero(Total[Idx], 1.f))
-			{
-				const float Percent = Delta[Idx] / Total[Idx];
-				const float AdditiveDelta = FMath::Abs(Additive[Idx]) * Percent;
-				CurrentAdditive[Idx] = (Additive[Idx] > 0.f) ? PreviousAdditive[Idx] + AdditiveDelta : PreviousAdditive[Idx] - AdditiveDelta;
-			}
-			else
-			{
-				CurrentAdditive[Idx] = Additive[Idx] * Alpha;
-			}
-		}
-		return CurrentAdditive;
-	};
-
-	Output.Initialize(BoneContainer.GetCompactPoseNumBones());
-
-	for (const FCompactPoseBoneIndex PoseBoneIndex : AdditivePose.GetPose().ForEachBoneIndex())
-	{
-		const FTransform& AdditiveTransform = AdditivePose.GetPose()[PoseBoneIndex];
-		if (AdditiveTransform.Equals(FTransform::Identity))
-		{
-			continue;
-		}
-
-		const int32 BoneIndex = PoseBoneIndex.GetInt();
-		const FMotionDeltaTrack& MotionDeltaTrack = MotionDeltaTracks.Tracks[BoneIndex];
-		const int32 TotalFrames = MotionDeltaTrack.BoneTransformTrack.Num();
-
-		FRawAnimSequenceTrack& Track = Output.AnimationTracks[BoneIndex];
-
-		Track.PosKeys.Reserve(TotalFrames);
-		Track.RotKeys.Reserve(TotalFrames);
-
-		Track.PosKeys.Add(MotionDeltaTrack.BoneTransformTrack[0].GetLocation());
-		Track.RotKeys.Add(MotionDeltaTrack.BoneTransformTrack[0].GetRotation());
-
-		Track.ScaleKeys.Add(FVector(1));
-
-		Output.TrackNames[BoneIndex] = BoneContainer.GetReferenceSkeleton().GetBoneName(BoneContainer.GetBoneIndicesArray()[BoneIndex]);
-
-		const FVector TotalAdditiveTranslation = AdditiveTransform.GetTranslation();
-		const FVector TotalAdditiveRotation = AdditiveTransform.GetRotation().Rotator().Euler();
+		InMotionWarpingComp->AddModifier(NewModifier);
 
-		FVector PrevAdditiveTranslation = FVector::ZeroVector;
-		FVector PrevAdditiveRotation = FVector::ZeroVector;
-		for (int32 FrameIdx = 1; FrameIdx < TotalFrames; FrameIdx++)
-		{
-			const FTransform& BoneTransform = MotionDeltaTrack.BoneTransformTrack[FrameIdx];
-			const float Alpha = (FrameIdx / (float)(TotalFrames - 1));
-
-			// Translation Key
-			const FVector DeltaTranslation = MotionDeltaTrack.DeltaTranslationTrack[FrameIdx];
-			const FVector CurrentAdditiveTranslation = CalculateAdditive(MotionDeltaTrack.TotalTranslation, DeltaTranslation, TotalAdditiveTranslation, PrevAdditiveTranslation, Alpha);
-			Output.AnimationTracks[BoneIndex].PosKeys.Add(BoneTransform.GetTranslation() + CurrentAdditiveTranslation);
-			PrevAdditiveTranslation = CurrentAdditiveTranslation;
-
-			// Rotation Key
-			const FVector DeltaRotation = MotionDeltaTrack.DeltaRotationTrack[FrameIdx].Euler();
-			const FVector CurrentAdditiveRotation = CalculateAdditive(MotionDeltaTrack.TotalRotation.Euler(), DeltaRotation, TotalAdditiveRotation, PrevAdditiveRotation, Alpha);
-			Output.AnimationTracks[BoneIndex].RotKeys.Add(FRotator::MakeFromEuler(BoneTransform.GetRotation().Rotator().Euler() + CurrentAdditiveRotation).Quaternion());
-			PrevAdditiveRotation = CurrentAdditiveRotation;
-		}
+		return NewModifier;
 	}
-}
-
-FTransform UMotionModifier_AdjustmentBlendWarp::ExtractWarpedRootMotion() const
-{
-	FTransform StartRootTransform;
-	ExtractBoneTransformAtTime(StartRootTransform, 0, PreviousPosition);
-
-	FTransform EndRootTransform;
-	ExtractBoneTransformAtTime(EndRootTransform, 0, CurrentPosition);
-
-	return EndRootTransform.GetRelativeTransform(StartRootTransform);
-}
 
-#if !(UE_BUILD_SHIPPING || UE_BUILD_TEST)
-void UMotionModifier_AdjustmentBlendWarp::DrawDebugWarpedTracks(float DrawDuration) const
-{
-	const ACharacter* CharacterOwner = GetCharacterOwner();
-	const UWorld* World = CharacterOwner ? CharacterOwner->GetWorld() : nullptr;
-	if (World && Result.GetNum() > 0 && PreviousPosition <= EndTime)
-	{
-		FTransform RootStartTransform;
-		ExtractBoneTransformAtFrame(RootStartTransform, 0, 0);
-
-		for (int32 TrackIndex = 0; TrackIndex < Result.GetNum(); TrackIndex++)
-		{
-			const int32 TotalFrames = Result.AnimationTracks[TrackIndex].PosKeys.Num();
-			if (TotalFrames > 1)
-			{
-				for (int32 FrameIdx = 0; FrameIdx < Result.AnimationTracks[TrackIndex].PosKeys.Num() - 1; FrameIdx++)
-				{
-					FTransform StartTransform;
-					ExtractBoneTransformAtFrame(StartTransform, TrackIndex, FrameIdx);
-					StartTransform = StartTransform * RootStartTransform.Inverse() * CachedMeshTransform;
-
-					FTransform EndTransform;
-					ExtractBoneTransformAtFrame(EndTransform, TrackIndex, FrameIdx + 1);
-					EndTransform = EndTransform * RootStartTransform.Inverse() * CachedMeshTransform;
-
-					DrawDebugLine(World, StartTransform.GetLocation(), EndTransform.GetLocation(), FColor::Yellow, false, DrawDuration, 0, 0.5f);
-				}
-			}
-		}
-
-		const FTransform& MeshTransform = CharacterOwner->GetMesh()->GetComponentTransform();
-		DrawDebugCoordinateSystem(World, CachedMeshTransform.GetLocation(), CachedMeshTransform.Rotator(), 20.f, false, DrawDuration, 0, 1.f);
-		DrawDebugCoordinateSystem(World, MeshTransform.GetLocation(), MeshTransform.Rotator(), 20.f, false, DrawDuration, 0, 1.f);
-		DrawDebugCoordinateSystem(World, GetTargetLocation(), GetTargetRotator(), 50.f, false, DrawDuration, 0, 1.f);
-		DrawDebugCoordinateSystem(World, (CachedRootMotion * CachedMeshTransform).GetLocation(), (CachedRootMotion * CachedMeshTransform).Rotator(), 50.f, false, DrawDuration, 0, 1.f);
-	}
-}
-#endif
-
-void UMotionModifier_AdjustmentBlendWarp::GetIKBoneTransformAndAlpha(FName BoneName, FTransform& OutTransform, float& OutAlpha) const
-{
-	if (Result.GetNum() == 0 || !bWarpIKBones || !IKBones.Contains(BoneName))
-	{
-		OutTransform = FTransform::Identity;
-		OutAlpha = 0.f;
-		return;
-	}
-
-	FTransform RootPrevPosition;
-	ExtractBoneTransformAtTime(RootPrevPosition, 0, 0.f);
-
-	FTransform BoneTransform;
-	ExtractBoneTransformAtTime(BoneTransform, BoneName, PreviousPosition);
-
-	OutTransform = BoneTransform * RootPrevPosition.Inverse() * CachedMeshTransform;
-	OutAlpha = Weight;
+	return nullptr;
 }
\ No newline at end of file
diff --git a/Engine/Plugins/Experimental/Animation/MotionWarping/Source/MotionWarping/Private/RootMotionModifier_SkewWarp.cpp b/Engine/Plugins/Experimental/Animation/MotionWarping/Source/MotionWarping/Private/RootMotionModifier_SkewWarp.cpp
--- a/Engine/Plugins/Experimental/Animation/MotionWarping/Source/MotionWarping/Private/RootMotionModifier_SkewWarp.cpp
+++ b/Engine/Plugins/Experimental/Animation/MotionWarping/Source/MotionWarping/Private/RootMotionModifier_SkewWarp.cpp
@@ -7,162 +7,6 @@
 #include "Components/SkeletalMeshComponent.h"
 #include "MotionWarpingComponent.h"
 
-FTransform FRootMotionModifier_SkewWarp::ProcessRootMotion(UMotionWarpingComponent& OwnerComp_DEPRECATED, const FTransform& InRootMotion, float DeltaSeconds)
-{
-	const ACharacter* CharacterOwner = GetCharacterOwner();
-	if(CharacterOwner == nullptr)
-	{
-		return InRootMotion;
-	}
-
-	FTransform FinalRootMotion = InRootMotion;
-
-	const FTransform RootMotionTotal = UMotionWarpingUtilities::ExtractRootMotionFromAnimation(Animation.Get(), PreviousPosition, EndTime);
-	const FTransform RootMotionDelta = UMotionWarpingUtilities::ExtractRootMotionFromAnimation(Animation.Get(), PreviousPosition, FMath::Min(CurrentPosition, EndTime));
-
-	if (bWarpTranslation && !RootMotionDelta.GetTranslation().IsNearlyZero())
-	{
-		const FTransform CurrentTransform = FTransform(
-			CharacterOwner->GetActorQuat(),
-			CharacterOwner->GetActorLocation() - FVector::UpVector * CharacterOwner->GetSimpleCollisionHalfHeight()
-		);
-
-		const FTransform MeshRelativeTransform = FTransform(CharacterOwner->GetBaseRotationOffset(), CharacterOwner->GetBaseTranslationOffset());
-		const FTransform MeshTransform = MeshRelativeTransform * CharacterOwner->GetActorTransform();
-		const FTransform RootMotionTotalWorldSpace = RootMotionTotal * MeshTransform;
-		const FTransform RootMotionDeltaWorldSpace = CharacterOwner->GetMesh()->ConvertLocalRootMotionToWorld(RootMotionDelta);
-
-		const FVector CurrentLocation = CurrentTransform.GetLocation();
-		const FQuat CurrentRotation = CurrentTransform.GetRotation();
-		
-		FVector TargetLocation = GetTargetLocation();
-		if(bIgnoreZAxis)
-		{
-			TargetLocation.Z = CurrentLocation.Z;
-		}
-
-		const FVector Translation = RootMotionDeltaWorldSpace.GetTranslation();
-		const FVector FutureLocation = RootMotionTotalWorldSpace.GetLocation();
-		const FVector CurrentToWorldOffset = TargetLocation - CurrentLocation;
-		const FVector CurrentToRootOffset = FutureLocation - CurrentLocation;
-
-		// Create a matrix we can use to put everything into a space looking straight at RootMotionSyncPosition. "forward" should be the axis along which we want to scale. 
-		FVector ToRootNormalized = CurrentToRootOffset.GetSafeNormal();
-
-		float BestMatchDot = FMath::Abs(FVector::DotProduct(ToRootNormalized, CurrentRotation.GetAxisX()));
-		FMatrix ToRootSyncSpace = FRotationMatrix::MakeFromXZ(ToRootNormalized, CurrentRotation.GetAxisZ());
-
-		float ZDot = FMath::Abs(FVector::DotProduct(ToRootNormalized, CurrentRotation.GetAxisZ()));
-		if (ZDot > BestMatchDot)
-		{
-			ToRootSyncSpace = FRotationMatrix::MakeFromXZ(ToRootNormalized, CurrentRotation.GetAxisX());
-			BestMatchDot = ZDot;
-		}
-
-		float YDot = FMath::Abs(FVector::DotProduct(ToRootNormalized, CurrentRotation.GetAxisY()));
-		if (YDot > BestMatchDot)
-		{
-			ToRootSyncSpace = FRotationMatrix::MakeFromXZ(ToRootNormalized, CurrentRotation.GetAxisZ());
-		}
-
-		// Put everything into RootSyncSpace.
-		const FVector RootMotionInSyncSpace = ToRootSyncSpace.InverseTransformVector(Translation);
-		const FVector CurrentToWorldSync = ToRootSyncSpace.InverseTransformVector(CurrentToWorldOffset);
-		const FVector CurrentToRootMotionSync = ToRootSyncSpace.InverseTransformVector(CurrentToRootOffset);
-
-		FVector CurrentToWorldSyncNorm = CurrentToWorldSync;
-		CurrentToWorldSyncNorm.Normalize();
-
-		FVector CurrentToRootMotionSyncNorm = CurrentToRootMotionSync;
-		CurrentToRootMotionSyncNorm.Normalize();
-
-		// Calculate skew Yaw Angle. 
-		FVector FlatToWorld = FVector(CurrentToWorldSyncNorm.X, CurrentToWorldSyncNorm.Y, 0.0f);
-		FlatToWorld.Normalize();
-		FVector FlatToRoot = FVector(CurrentToRootMotionSyncNorm.X, CurrentToRootMotionSyncNorm.Y, 0.0f);
-		FlatToRoot.Normalize();
-		float AngleAboutZ = FMath::Acos(FVector::DotProduct(FlatToWorld, FlatToRoot));
-		float AngleAboutZNorm = FMath::DegreesToRadians(FRotator::NormalizeAxis(FMath::RadiansToDegrees(AngleAboutZ)));
-		if (FlatToWorld.Y < 0.0f)
-		{
-			AngleAboutZNorm *= -1.0f;
-		}
-
-		// Calculate Skew Pitch Angle. 
-		FVector ToWorldNoY = FVector(CurrentToWorldSyncNorm.X, 0.0f, CurrentToWorldSyncNorm.Z);
-		ToWorldNoY.Normalize();
-		FVector ToRootNoY = FVector(CurrentToRootMotionSyncNorm.X, 0.0f, CurrentToRootMotionSyncNorm.Z);
-		ToRootNoY.Normalize();
-		const float AngleAboutY = FMath::Acos(FVector::DotProduct(ToWorldNoY, ToRootNoY));
-		float AngleAboutYNorm = FMath::DegreesToRadians(FRotator::NormalizeAxis(FMath::RadiansToDegrees(AngleAboutY)));
-		if (ToWorldNoY.Z < 0.0f)
-		{
-			AngleAboutYNorm *= -1.0f;
-		}
-
-		FVector SkewedRootMotion = FVector::ZeroVector;
-		float ProjectedScale = FVector::DotProduct(CurrentToWorldSync, CurrentToRootMotionSyncNorm) / CurrentToRootMotionSync.Size();
-		if (ProjectedScale != 0.0f)
-		{
-			FMatrix ScaleMatrix;
-			ScaleMatrix.SetIdentity();
-			ScaleMatrix.SetAxis(0, FVector(ProjectedScale, 0.0f, 0.0f));
-			ScaleMatrix.SetAxis(1, FVector(0.0f, 1.0f, 0.0f));
-			ScaleMatrix.SetAxis(2, FVector(0.0f, 0.0f, 1.0f));
-
-			FMatrix ShearXAlongYMatrix;
-			ShearXAlongYMatrix.SetIdentity();
-			ShearXAlongYMatrix.SetAxis(0, FVector(1.0f, FMath::Tan(AngleAboutZNorm), 0.0f));
-			ShearXAlongYMatrix.SetAxis(1, FVector(0.0f, 1.0f, 0.0f));
-			ShearXAlongYMatrix.SetAxis(2, FVector(0.0f, 0.0f, 1.0f));
-
-			FMatrix ShearXAlongZMatrix;
-			ShearXAlongZMatrix.SetIdentity();
-			ShearXAlongZMatrix.SetAxis(0, FVector(1.0f, 0.0f, FMath::Tan(AngleAboutYNorm)));
-			ShearXAlongZMatrix.SetAxis(1, FVector(0.0f, 1.0f, 0.0f));
-			ShearXAlongZMatrix.SetAxis(2, FVector(0.0f, 0.0f, 1.0f));
-
-			FMatrix ScaledSkewMatrix = ScaleMatrix * ShearXAlongYMatrix * ShearXAlongZMatrix;
-
-			// Skew and scale the Root motion. 
-			SkewedRootMotion = ScaledSkewMatrix.TransformVector(RootMotionInSyncSpace);
-		}
-		else if (!CurrentToRootMotionSync.IsZero() && !CurrentToWorldSync.IsZero() && !RootMotionInSyncSpace.IsZero())
-		{
-			// Figure out ratio between remaining Root and remaining World. Then project scaled length of current Root onto World.
-			const float Scale = CurrentToWorldSync.Size() / CurrentToRootMotionSync.Size();
-			const float StepTowardTarget = RootMotionInSyncSpace.ProjectOnTo(RootMotionInSyncSpace).Size();
-			SkewedRootMotion = CurrentToWorldSyncNorm * (Scale * StepTowardTarget);
-		}
-
-		// Put our result back in world space.  
-		FinalRootMotion.SetTranslation(ToRootSyncSpace.TransformVector(SkewedRootMotion));
-	}
-
-	if(bWarpRotation)
-	{
-		const FQuat WarpedRotation = WarpRotation(InRootMotion, RootMotionTotal, DeltaSeconds);
-		FinalRootMotion.SetRotation(WarpedRotation);
-	}
-
-	// Debug
-#if !(UE_BUILD_SHIPPING || UE_BUILD_TEST)
-	const int32 DebugLevel = FMotionWarpingCVars::CVarMotionWarpingDebug.GetValueOnGameThread();
-	if (DebugLevel > 0)
-	{
-		PrintLog(TEXT("FRootMotionModifier_Skew"), InRootMotion, FinalRootMotion);
-
-		if (DebugLevel >= 2)
-		{
-			const float DrawDebugDuration = FMotionWarpingCVars::CVarMotionWarpingDrawDebugDuration.GetValueOnGameThread();
-			DrawDebugCoordinateSystem(CharacterOwner->GetWorld(), GetTargetLocation(), GetTargetRotator(), 50.f, false, DrawDebugDuration, 0, 1.f);
-		}
-	}
-#endif
-
-	return FinalRootMotion;
-}
-
 // URootMotionModifierConfig_SkewWarp
 ///////////////////////////////////////////////////////////////
 
@@ -172,7 +16,7 @@ FRootMotionModifierHandle URootMotionModifierConfig_SkewWarp::AddRootMotionModif
 {
 	if (ensureAlways(InMotionWarpingComp))
 	{
-		UMotionModifier_SkewWarp* NewModifier = NewObject<UMotionModifier_SkewWarp>(InMotionWarpingComp);
+		URootMotionModifier_SkewWarp* NewModifier = NewObject<URootMotionModifier_SkewWarp>(InMotionWarpingComp);
 		NewModifier->Animation = InAnimation;
 		NewModifier->StartTime = InStartTime;
 		NewModifier->EndTime = InEndTime;
@@ -193,12 +37,12 @@ FRootMotionModifierHandle URootMotionModifierConfig_SkewWarp::AddRootMotionModif
 
 ///////////////////////////////////////////////////////////////
 
-UMotionModifier_SkewWarp::UMotionModifier_SkewWarp(const FObjectInitializer& ObjectInitializer)
+URootMotionModifier_SkewWarp::URootMotionModifier_SkewWarp(const FObjectInitializer& ObjectInitializer)
 	: Super(ObjectInitializer)
 {
 }
 
-FTransform UMotionModifier_SkewWarp::ProcessRootMotion(const FTransform& InRootMotion, float DeltaSeconds)
+FTransform URootMotionModifier_SkewWarp::ProcessRootMotion(const FTransform& InRootMotion, float DeltaSeconds)
 {
 	const ACharacter* CharacterOwner = GetCharacterOwner();
 	if(CharacterOwner == nullptr)
@@ -352,4 +196,32 @@ FTransform UMotionModifier_SkewWarp::ProcessRootMotion(const FTransform& InRootM
 #endif
 
 	return FinalRootMotion;
+}
+
+URootMotionModifier_SkewWarp* URootMotionModifier_SkewWarp::AddRootMotionModifierSkewWarp(UMotionWarpingComponent* InMotionWarpingComp, const UAnimSequenceBase* InAnimation, float InStartTime, float InEndTime,
+	FName InWarpTargetName, EWarpPointAnimProvider InWarpPointAnimProvider, FTransform InWarpPointAnimTransform, FName InWarpPointAnimBoneName,
+	bool bInWarpTranslation, bool bInIgnoreZAxis, bool bInWarpRotation, EMotionWarpRotationType InRotationType, float InWarpRotationTimeMultiplier)
+{
+	if (ensureAlways(InMotionWarpingComp))
+	{
+		URootMotionModifier_SkewWarp* NewModifier = NewObject<URootMotionModifier_SkewWarp>(InMotionWarpingComp);
+		NewModifier->Animation = InAnimation;
+		NewModifier->StartTime = InStartTime;
+		NewModifier->EndTime = InEndTime;
+		NewModifier->WarpTargetName = InWarpTargetName;
+		NewModifier->WarpPointAnimProvider = InWarpPointAnimProvider;
+		NewModifier->WarpPointAnimTransform = InWarpPointAnimTransform;
+		NewModifier->WarpPointAnimBoneName = InWarpPointAnimBoneName;
+		NewModifier->bWarpTranslation = bInWarpTranslation;
+		NewModifier->bIgnoreZAxis = bInIgnoreZAxis;
+		NewModifier->bWarpRotation = bInWarpRotation;
+		NewModifier->RotationType = InRotationType;
+		NewModifier->WarpRotationTimeMultiplier = InWarpRotationTimeMultiplier;
+		
+		InMotionWarpingComp->AddModifier(NewModifier);
+
+		return NewModifier;
+	}
+
+	return nullptr;
 }
\ No newline at end of file
diff --git a/Engine/Plugins/Experimental/Animation/MotionWarping/Source/MotionWarping/Public/AnimNotifyState_MotionWarping.h b/Engine/Plugins/Experimental/Animation/MotionWarping/Source/MotionWarping/Public/AnimNotifyState_MotionWarping.h
--- a/Engine/Plugins/Experimental/Animation/MotionWarping/Source/MotionWarping/Public/AnimNotifyState_MotionWarping.h
+++ b/Engine/Plugins/Experimental/Animation/MotionWarping/Source/MotionWarping/Public/AnimNotifyState_MotionWarping.h
@@ -21,10 +21,7 @@ public:
 	//@TODO: Prevent notify callbacks and add comments explaining why we don't use those here.
 
 	UPROPERTY(EditAnywhere, Instanced, BlueprintReadWrite, Category = "Config")
-	URootMotionModifierConfig* RootMotionModifierConfig;
-
-	UPROPERTY(EditAnywhere, Instanced, BlueprintReadWrite, Category = "Config")
-	UMotionModifier* RootMotionModifier;
+	URootMotionModifier* RootMotionModifier;
 
 	UAnimNotifyState_MotionWarping(const FObjectInitializer& ObjectInitializer);
 
diff --git a/Engine/Plugins/Experimental/Animation/MotionWarping/Source/MotionWarping/Public/MotionWarpingComponent.h b/Engine/Plugins/Experimental/Animation/MotionWarping/Source/MotionWarping/Public/MotionWarpingComponent.h
--- a/Engine/Plugins/Experimental/Animation/MotionWarping/Source/MotionWarping/Public/MotionWarpingComponent.h
+++ b/Engine/Plugins/Experimental/Animation/MotionWarping/Source/MotionWarping/Public/MotionWarpingComponent.h
@@ -74,7 +74,7 @@ public:
 
 	/** @return All the MotionWarping windows within the supplied animation for a given Sync Point */
 	UFUNCTION(BlueprintCallable, Category = "Motion Warping")
-	static void GetMotionWarpingWindowsForSyncPointFromAnimation(const UAnimSequenceBase* Animation, FName SyncPointName, TArray<FMotionWarpingWindowData>& OutWindows);
+	static void GetMotionWarpingWindowsForSyncPointFromAnimation(const UAnimSequenceBase* Animation, FName WarpTargetName, TArray<FMotionWarpingWindowData>& OutWindows);
 };
 
 DECLARE_DYNAMIC_MULTICAST_DELEGATE_OneParam(FMotionWarpingPreUpdate, class UMotionWarpingComponent*, MotionWarpingComp);
@@ -102,10 +102,7 @@ public:
 	FORCEINLINE ACharacter* GetCharacterOwner() const { return CharacterOwner.Get(); }
 
 	/** Returns the list of root motion modifiers */
-	FORCEINLINE const TArray<TSharedPtr<FRootMotionModifier>>& GetRootMotionModifiers() const { return RootMotionModifiers; }
-
-	/** Returns the list of root motion modifiers */
-	FORCEINLINE const TArray<UMotionModifier*>& GetModifiers() const { return Modifiers; }
+	FORCEINLINE const TArray<URootMotionModifier*>& GetModifiers() const { return Modifiers; }
 
 	/** Find the SyncPoint associated with a specified name */
 	FORCEINLINE const FMotionWarpingSyncPoint* FindSyncPoint(const FName& SyncPointName) const { return SyncPoints.Find(SyncPointName); }
@@ -122,21 +119,15 @@ public:
 	bool ContainsModifier(const UAnimSequenceBase* Animation, float StartTime, float EndTime) const;
 
 	/** Add a new modifier */
-	FRootMotionModifierHandle AddRootMotionModifier(TSharedPtr<FRootMotionModifier> Modifier);
-
-	/** Add a new modifier */
-	FRootMotionModifierHandle AddModifier(UMotionModifier* Modifier);
+	FRootMotionModifierHandle AddModifier(URootMotionModifier* Modifier);
 
 	/** Mark all the modifiers as Disable */
 	void DisableAllRootMotionModifiers();
 
 	/** Finds a RootMotionModifier that matches the supplied Handle */
-	TSharedPtr<FRootMotionModifier> GetRootMotionModifierByHandle(const FRootMotionModifierHandle& Handle) const;
+	URootMotionModifier* GetModifierByHandle(const FRootMotionModifierHandle& Handle) const;
 
-	/** Finds a RootMotionModifier that matches the supplied Handle */
-	UMotionModifier* GetModifierByHandle(const FRootMotionModifierHandle& Handle) const;
-
-	UMotionModifier* AddModifierFromTemplate(UMotionModifier* Template, const UAnimSequenceBase* Animation, float StartTime, float EndTime);
+	URootMotionModifier* AddModifierFromTemplate(URootMotionModifier* Template, const UAnimSequenceBase* Animation, float StartTime, float EndTime);
 
 protected:
 
@@ -144,12 +135,9 @@ protected:
 	UPROPERTY(Transient)
 	TWeakObjectPtr<ACharacter> CharacterOwner;
 
-	/** List of root motion modifiers */
-	TArray<TSharedPtr<FRootMotionModifier>> RootMotionModifiers;
-
 	/** List of root motion modifiers */
 	UPROPERTY(Transient)
-	TArray<UMotionModifier*> Modifiers;
+	TArray<URootMotionModifier*> Modifiers;
 
 	UPROPERTY(Transient)
 	TMap<FName, FMotionWarpingSyncPoint> SyncPoints;
diff --git a/Engine/Plugins/Experimental/Animation/MotionWarping/Source/MotionWarping/Public/RootMotionModifier.h b/Engine/Plugins/Experimental/Animation/MotionWarping/Source/MotionWarping/Public/RootMotionModifier.h
--- a/Engine/Plugins/Experimental/Animation/MotionWarping/Source/MotionWarping/Public/RootMotionModifier.h
+++ b/Engine/Plugins/Experimental/Animation/MotionWarping/Source/MotionWarping/Public/RootMotionModifier.h
@@ -56,135 +56,6 @@ private:
 
 DECLARE_DYNAMIC_DELEGATE_TwoParams(FOnRootMotionModifierDelegate, UMotionWarpingComponent*, MotionWarpingComp, const FRootMotionModifierHandle&, Handle);
 
-/** Base struct for Root Motion Modifiers */
-USTRUCT()
-struct MOTIONWARPING_API FRootMotionModifier
-{
-	GENERATED_BODY()
-
-public:
-
-	/** Source of the root motion we are warping */
-	UPROPERTY()
-	TWeakObjectPtr<const UAnimSequenceBase> Animation = nullptr;
-
-	/** Start time of the warping window */
-	UPROPERTY()
-	float StartTime = 0.f;
-
-	/** End time of the warping window */
-	UPROPERTY()
-	float EndTime = 0.f;
-
-	/** Previous playback time of the animation */
-	UPROPERTY()
-	float PreviousPosition = 0.f;
-
-	/** Current playback time of the animation */
-	UPROPERTY()
-	float CurrentPosition = 0.f;
-
-	/** Current blend weight of the animation */
-	UPROPERTY()
-	float Weight = 0.f;
-
-	/** Whether this modifier runs before the extracted root motion is converted to world space or after */
-	UPROPERTY()
-	bool bInLocalSpace = false;
-
-	/** Current state */
-	UPROPERTY()
-	ERootMotionModifierState State = ERootMotionModifierState::Waiting; //@TODO: Move to private
-
-	/** Delegate called when this modifier is activated (starts affecting the root motion) */
-	UPROPERTY()
-	FOnRootMotionModifierDelegate OnActivateDelegate;
-
-	/** Delegate called when this modifier updates while active (affecting the root motion) */
-	UPROPERTY()
-	FOnRootMotionModifierDelegate OnUpdateDelegate;
-
-	/** Delegate called when this modifier is deactivated (stops affecting the root motion) */
-	UPROPERTY()
-	FOnRootMotionModifierDelegate OnDeactivateDelegate;
-
-	FRootMotionModifier(){}
-	virtual ~FRootMotionModifier(){}
-
-	/** Initialize this modifier */
-	void Initialize(UMotionWarpingComponent* OwnerComp);
-
-	/** Returns the actual struct used for RTTI */
-	virtual UScriptStruct* GetScriptStruct() const PURE_VIRTUAL(FRootMotionModifier::GetScriptStruct, return FRootMotionModifier::StaticStruct(););
-
-	/** Called when the state of the modifier changes */
-	virtual void OnStateChanged(ERootMotionModifierState LastState);
-
-	/** Sets the state of the modifier */
-	void SetState(ERootMotionModifierState NewState);
-
-	/** Returns the state of the modifier */
-	FORCEINLINE ERootMotionModifierState GetState() const { return State; }
-
-	/** Returns the handle to identify this modifier */
-	FORCEINLINE const FRootMotionModifierHandle& GetHandle() const { return Handle; }
-
-	/** Returns a pointer to the component that owns this modifier */
-	FORCEINLINE UMotionWarpingComponent* GetOwnerComponent() const { return OwnerComponent.Get(); }
-
-	/** Returns a pointer to the character that owns the component that owns this modifier */
-	class ACharacter* GetCharacterOwner() const;
-
-	//DEPRECATED will be removed shortly, replaced with a signature that does not receive the comp via param.
-	virtual void Update(UMotionWarpingComponent& OwnerComp);
-	virtual FTransform ProcessRootMotion(UMotionWarpingComponent& OwnerComp, const FTransform& InRootMotion, float DeltaSeconds) { return FTransform::Identity;	}
-
-	FORCEINLINE const UAnimSequenceBase* GetAnimation() const { return Animation.Get(); }
-
-private:
-
-	friend UMotionWarpingComponent;
-
-	/** Pointer back to the Component that owns this modifier */
-	UPROPERTY()
-	TWeakObjectPtr<UMotionWarpingComponent> OwnerComponent = nullptr;
-
-	/** Handle to identify this modifier */
-	UPROPERTY()
-	FRootMotionModifierHandle Handle;
-
-};
-
-// Mark FRootMotionModifier as pure virtual, so that the UObject machinery won't attempt to instantiate it.
-template<>
-struct TStructOpsTypeTraits<FRootMotionModifier> :
-	TStructOpsTypeTraitsBase2<FRootMotionModifier>
-{
-	enum
-	{
-		WithPureVirtual = true,
-    };
-};
-
-
-/** Blueprint wrapper around the config properties of a root motion modifier */
-UCLASS(abstract, BlueprintType, EditInlineNew)
-class MOTIONWARPING_API URootMotionModifierConfig : public UObject
-{
-	GENERATED_BODY()
-
-public:
-	
-	URootMotionModifierConfig(const FObjectInitializer& ObjectInitializer)
-		: Super(ObjectInitializer) {}
-
-	/** Adds a RootMotionModifier of the type this object represents */
-	virtual FRootMotionModifierHandle AddRootMotionModifierNew(UMotionWarpingComponent* InMotionWarpingComp, const UAnimSequenceBase* Animation, float StartTime, float EndTime) const { return FRootMotionModifierHandle::InvalidHandle; }
-
-	//DEPRECATED will be removed shortly
-	virtual void AddRootMotionModifier(UMotionWarpingComponent* InMotionWarpingComp, const UAnimSequenceBase* Animation, float StartTime, float EndTime) const {}
-};
-
 /** Represents a point of alignment in the world */
 USTRUCT(BlueprintType, meta = (HasNativeMake = "MotionWarping.MotionWarpingUtilities.MakeMotionWarpingSyncPoint", HasNativeBreak = "MotionWarping.MotionWarpingUtilities.BreakMotionWarpingSyncPoint"))
 struct MOTIONWARPING_API FMotionWarpingSyncPoint
@@ -250,207 +121,8 @@ enum class EWarpPointAnimProvider : uint8
 	Bone
 };
 
-USTRUCT()
-struct MOTIONWARPING_API FRootMotionModifier_Warp : public FRootMotionModifier
-{
-	GENERATED_BODY()
-
-public:
-
-	/** Name used to find the warp point for this modifier */
-	UPROPERTY()
-	FName SyncPointName = NAME_None; //@TODO: Rename to WarpPointName
-
-	/** Method used to extract the warp point from the animation */
-	UPROPERTY()
-	EWarpPointAnimProvider WarpPointAnimProvider = EWarpPointAnimProvider::None;
-
-	UPROPERTY()
-	FTransform WarpPointAnimTransform = FTransform::Identity;
-
-	UPROPERTY()
-	FName WarpPointAnimBoneName = NAME_None;
-
-	/** Whether to warp the translation component of the root motion */
-	UPROPERTY()
-	bool bWarpTranslation = true;
-
-	/** Whether to ignore the Z component of the translation. Z motion will remain untouched */
-	UPROPERTY()
-	bool bIgnoreZAxis = true;
-
-	/** Whether to warp the rotation component of the root motion */
-	UPROPERTY()
-	bool bWarpRotation = true;
-
-	/** Whether rotation should be warp to match the rotation of the sync point or to face the sync point */
-	UPROPERTY()
-	EMotionWarpRotationType RotationType = EMotionWarpRotationType::Default;
-
-	/**
-	 * Allow to modify how fast the rotation is warped.
-	 * e.g if the window duration is 2sec and this is 0.5, the target rotation will be reached in 1sec instead of 2sec
-	 */
-	UPROPERTY()
-	float WarpRotationTimeMultiplier = 1.f;
-
-	//DEPRECATED: Will be removed shortly
-	UPROPERTY()
-	FMotionWarpingSyncPoint CachedSyncPoint;
-
-	FRootMotionModifier_Warp(){}
-	virtual ~FRootMotionModifier_Warp() {}
-
-	//~ Begin FRootMotionModifier Interface
-	virtual UScriptStruct* GetScriptStruct() const override { return FRootMotionModifier_Warp::StaticStruct(); }
-	virtual void Update(UMotionWarpingComponent& OwnerComp) override;
-	virtual FTransform ProcessRootMotion(UMotionWarpingComponent& OwnerComp, const FTransform& InRootMotion, float DeltaSeconds) override;
-	//~ End FRootMotionModifier Interface
-
-	/** Event called during update if the target transform changes while the warping is active */
-	virtual void OnTargetTransformChanged() {}
-
-#if !(UE_BUILD_SHIPPING || UE_BUILD_TEST)
-	void PrintLog(const FString& Name, const FTransform& OriginalRootMotion, const FTransform& WarpedRootMotion) const;
-
-	//DEPRECATED will be removed shortly
-	void PrintLog(UMotionWarpingComponent& OwnerComp, const FString& Name, const FTransform& OriginalRootMotion, const FTransform& WarpedRootMotion) const
-	{
-		PrintLog(Name, OriginalRootMotion, WarpedRootMotion);
-	}
-
-#endif
-
-protected:
-
-	FORCEINLINE FVector GetTargetLocation() const { return CachedTargetTransform.GetLocation(); }
-	FORCEINLINE FRotator GetTargetRotator() const { return GetTargetRotation().Rotator(); }
-	FQuat GetTargetRotation() const;
-	
-	FQuat WarpRotation(const FTransform& RootMotionDelta, const FTransform& RootMotionTotal, float DeltaSeconds);
-
-	//@TODO: This should be Optional and private
-	UPROPERTY()
-	FTransform CachedTargetTransform = FTransform::Identity;
-
-private:
-
-	/** Cached of the offset from the warp point. Used to calculate the final target transform when a warp point is defined in the animation */
-	TOptional<FTransform> CachedOffsetFromWarpPoint;
-};
-
-UCLASS(meta = (DisplayName = "Simple Warp"))
-class MOTIONWARPING_API URootMotionModifierConfig_Warp : public URootMotionModifierConfig
-{
-	GENERATED_BODY()
-
-public:
-
-	/** Name used to find the warp point for this modifier */
-	UPROPERTY(EditAnywhere, BlueprintReadWrite, Category = "Config", meta = (DisplayName = "Warp Point Name"))
-	FName SyncPointName = NAME_None;  //@TODO: Rename to WarpPointName
-
-	UPROPERTY(EditAnywhere, BlueprintReadWrite, Category = "Config")
-	EWarpPointAnimProvider WarpPointAnimProvider = EWarpPointAnimProvider::None;
-
-	//@TODO: Hide from the UI when Provider != Static
-	UPROPERTY(EditAnywhere, BlueprintReadWrite, Category = "Config", meta = (EditCondition = "WarpPointAnimProvider == EWarpPointAnimProvider::Static"))
-	FTransform WarpPointAnimTransform = FTransform::Identity;
-
-	//@TODO: Hide from the UI when Provider != Bone
-	UPROPERTY(EditAnywhere, BlueprintReadWrite, Category = "Config", meta = (EditCondition = "WarpPointAnimProvider == EWarpPointAnimProvider::Bone"))
-	FName WarpPointAnimBoneName = NAME_None;
-
-	/** Whether to warp the translation component of the root motion */
-	UPROPERTY(EditAnywhere, BlueprintReadWrite, Category = "Config")
-	bool bWarpTranslation = true;
-
-	/** Whether to ignore the Z component of the translation. Z motion will remain untouched */
-	UPROPERTY(EditAnywhere, BlueprintReadWrite, Category = "Config", meta = (EditCondition = "bWarpTranslation"))
-	bool bIgnoreZAxis = true;
-
-	/** Whether to warp the rotation component of the root motion */
-	UPROPERTY(EditAnywhere, BlueprintReadWrite, Category = "Config")
-	bool bWarpRotation = true;
-
-	/** Whether rotation should be warp to match the rotation of the sync point or to face the sync point */
-	UPROPERTY(EditAnywhere, BlueprintReadWrite, Category = "Config", meta = (EditCondition = "bWarpRotation"))
-	EMotionWarpRotationType RotationType;
-
-	/**
-	 * Allow to modify how fast the rotation is warped.
-	 * e.g if the window duration is 2sec and this is 0.5, the target rotation will be reached in 1sec instead of 2sec
-	 */
-	UPROPERTY(EditAnywhere, BlueprintReadWrite, Category = "Config", meta = (EditCondition = "bWarpRotation"))
-	float WarpRotationTimeMultiplier = 1.f;
-
-	URootMotionModifierConfig_Warp(const FObjectInitializer& ObjectInitializer) : Super(ObjectInitializer) {}
-
-	virtual FRootMotionModifierHandle AddRootMotionModifierNew(UMotionWarpingComponent* MotionWarpingComp, const UAnimSequenceBase* Animation, float StartTime, float EndTime) const override
-	{
-		return URootMotionModifierConfig_Warp::AddRootMotionModifierSimpleWarp(MotionWarpingComp, Animation, StartTime, EndTime, 
-			SyncPointName, WarpPointAnimProvider, WarpPointAnimTransform, WarpPointAnimBoneName,
-			bWarpTranslation, bIgnoreZAxis, bWarpRotation, RotationType, WarpRotationTimeMultiplier);
-	}
-
-	UFUNCTION(BlueprintCallable, Category = "Motion Warping")
-	static FRootMotionModifierHandle AddRootMotionModifierSimpleWarp(UMotionWarpingComponent* InMotionWarpingComp, const UAnimSequenceBase* InAnimation, float InStartTime, float InEndTime,
-		FName InSyncPointName, EWarpPointAnimProvider InWarpPointAnimProvider, FTransform InWarpPointAnimTransform, FName InWarpPointAnimBoneName,
-		bool bInWarpTranslation, bool bInIgnoreZAxis, bool bInWarpRotation, EMotionWarpRotationType InRotationType, float InWarpRotationTimeMultiplier = 1.f);
-};
-
-// FRootMotionModifier_Scale
-///////////////////////////////////////////////////////////////
-
-USTRUCT()
-struct MOTIONWARPING_API FRootMotionModifier_Scale : public FRootMotionModifier
-{
-	GENERATED_BODY()
-
-public:
-
-	/** Vector used to scale each component of the translation */
-	UPROPERTY()
-	FVector Scale = FVector(1.f);
-
-	FRootMotionModifier_Scale() { bInLocalSpace = true; }
-	virtual ~FRootMotionModifier_Scale() {}
-
-	virtual UScriptStruct* GetScriptStruct() const override { return FRootMotionModifier_Scale::StaticStruct(); }
-	virtual FTransform ProcessRootMotion(UMotionWarpingComponent& OwnerComp, const FTransform& InRootMotion, float DeltaSeconds) override
-	{
-		FTransform FinalRootMotion = InRootMotion;
-		FinalRootMotion.ScaleTranslation(Scale);
-		return FinalRootMotion;
-	}
-};
-
-UCLASS(meta = (DisplayName = "Scale"))
-class MOTIONWARPING_API URootMotionModifierConfig_Scale : public URootMotionModifierConfig
-{
-	GENERATED_BODY()
-
-public:
-
-	/** Vector used to scale each component of the translation */
-	UPROPERTY(EditAnywhere, BlueprintReadWrite, Category = "Config")
-	FVector Scale = FVector(1.f);
-
-	URootMotionModifierConfig_Scale(const FObjectInitializer& ObjectInitializer) : Super(ObjectInitializer) {}
-
-	virtual FRootMotionModifierHandle AddRootMotionModifierNew(UMotionWarpingComponent* MotionWarpingComp, const UAnimSequenceBase* Animation, float StartTime, float EndTime) const override
-	{
-		return URootMotionModifierConfig_Scale::AddRootMotionModifierScale(MotionWarpingComp, Animation, StartTime, EndTime, Scale);
-	}
-
-	UFUNCTION(BlueprintCallable, Category = "Motion Warping")
-	static FRootMotionModifierHandle AddRootMotionModifierScale(UMotionWarpingComponent* InMotionWarpingComp, const UAnimSequenceBase* InAnimation, float InStartTime, float InEndTime, FVector InScale);
-};
-
-///////////////////////////////////////////////////////////////
-
 UCLASS(Abstract, BlueprintType, EditInlineNew)
-class MOTIONWARPING_API UMotionModifier : public UObject
+class MOTIONWARPING_API URootMotionModifier : public UObject
 {
 	GENERATED_BODY()
 
@@ -496,7 +168,7 @@ public:
 	UPROPERTY()
 	FOnRootMotionModifierDelegate OnDeactivateDelegate;
 
-	UMotionModifier(const FObjectInitializer& ObjectInitializer);
+	URootMotionModifier(const FObjectInitializer& ObjectInitializer);
 
 	/** Initialize this modifier */
 	void Initialize(UMotionWarpingComponent* OwnerComp);
@@ -511,7 +183,8 @@ public:
 	FORCEINLINE ERootMotionModifierState GetState() const { return State; }
 
 	/** Returns the handle to identify this modifier */
-	FORCEINLINE const FRootMotionModifierHandle& GetHandle() const { return Handle; }
+	UFUNCTION(BlueprintPure, Category = "Motion Warping")
+	FRootMotionModifierHandle GetHandle() const { return Handle; }
 
 	/** Returns a pointer to the component that owns this modifier */
 	FORCEINLINE UMotionWarpingComponent* GetOwnerComponent() const { return OwnerComponent.Get(); }
@@ -542,7 +215,7 @@ private:
 };
 
 UCLASS()
-class MOTIONWARPING_API UMotionModifier_Warp : public UMotionModifier
+class MOTIONWARPING_API URootMotionModifier_Warp : public URootMotionModifier
 {
 	GENERATED_BODY()
 
@@ -586,7 +259,7 @@ public:
 	UPROPERTY(EditAnywhere, BlueprintReadWrite, Category = "Config", meta = (EditCondition = "bWarpRotation"))
 	float WarpRotationTimeMultiplier = 1.f;
 
-	UMotionModifier_Warp(const FObjectInitializer& ObjectInitializer);
+	URootMotionModifier_Warp(const FObjectInitializer& ObjectInitializer);
 
 	//~ Begin FRootMotionModifier Interface
 	virtual void Update() override;
@@ -600,6 +273,11 @@ public:
 	void PrintLog(const FString& Name, const FTransform& OriginalRootMotion, const FTransform& WarpedRootMotion) const;
 #endif
 
+	UFUNCTION(BlueprintCallable, Category = "Motion Warping")
+	static URootMotionModifier_Warp* AddRootMotionModifierSimpleWarp(UMotionWarpingComponent* InMotionWarpingComp, const UAnimSequenceBase* InAnimation, float InStartTime, float InEndTime,
+		FName InWarpTargetName, EWarpPointAnimProvider InWarpPointAnimProvider, FTransform InWarpPointAnimTransform, FName InWarpPointAnimBoneName,
+		bool bInWarpTranslation, bool bInIgnoreZAxis, bool bInWarpRotation, EMotionWarpRotationType InRotationType, float InWarpRotationTimeMultiplier = 1.f);
+
 protected:
 
 	FORCEINLINE FVector GetTargetLocation() const { return CachedTargetTransform.GetLocation(); }
@@ -616,4 +294,103 @@ private:
 
 	/** Cached of the offset from the warp target. Used to calculate the final target transform when a warp target is defined in the animation */
 	TOptional<FTransform> CachedOffsetFromWarpPoint;
-};
\ No newline at end of file
+};
+
+UCLASS()
+class MOTIONWARPING_API URootMotionModifier_Scale : public URootMotionModifier
+{
+	GENERATED_BODY()
+
+public:
+
+	/** Vector used to scale each component of the translation */
+	UPROPERTY(EditAnywhere, BlueprintReadWrite, Category = "Config")
+	FVector Scale = FVector(1.f);
+
+	URootMotionModifier_Scale(const FObjectInitializer& ObjectInitializer) : Super(ObjectInitializer) { bInLocalSpace = true; }
+
+	virtual FTransform ProcessRootMotion(const FTransform& InRootMotion, float DeltaSeconds) override
+	{
+		FTransform FinalRootMotion = InRootMotion;
+		FinalRootMotion.ScaleTranslation(Scale);
+		return FinalRootMotion;
+	}
+};
+
+/** Blueprint wrapper around the config properties of a root motion modifier */
+UCLASS(abstract, BlueprintType, EditInlineNew)
+class MOTIONWARPING_API URootMotionModifierConfig : public UObject
+{
+	GENERATED_BODY()
+
+public:
+
+	URootMotionModifierConfig(const FObjectInitializer& ObjectInitializer)
+		: Super(ObjectInitializer) {}
+
+	/** Adds a RootMotionModifier of the type this object represents */
+	virtual FRootMotionModifierHandle AddRootMotionModifierNew(UMotionWarpingComponent* InMotionWarpingComp, const UAnimSequenceBase* Animation, float StartTime, float EndTime) const { return FRootMotionModifierHandle::InvalidHandle; }
+
+	//DEPRECATED will be removed shortly
+	virtual void AddRootMotionModifier(UMotionWarpingComponent* InMotionWarpingComp, const UAnimSequenceBase* Animation, float StartTime, float EndTime) const {}
+};
+
+UCLASS(meta = (DisplayName = "Simple Warp"))
+class MOTIONWARPING_API URootMotionModifierConfig_Warp : public URootMotionModifierConfig
+{
+	GENERATED_BODY()
+
+public:
+
+	/** Name used to find the warp point for this modifier */
+	UPROPERTY(EditAnywhere, BlueprintReadWrite, Category = "Config", meta = (DisplayName = "Warp Point Name"))
+	FName SyncPointName = NAME_None;  //@TODO: Rename to WarpPointName
+
+	UPROPERTY(EditAnywhere, BlueprintReadWrite, Category = "Config")
+	EWarpPointAnimProvider WarpPointAnimProvider = EWarpPointAnimProvider::None;
+
+	//@TODO: Hide from the UI when Provider != Static
+	UPROPERTY(EditAnywhere, BlueprintReadWrite, Category = "Config", meta = (EditCondition = "WarpPointAnimProvider == EWarpPointAnimProvider::Static"))
+	FTransform WarpPointAnimTransform = FTransform::Identity;
+
+	//@TODO: Hide from the UI when Provider != Bone
+	UPROPERTY(EditAnywhere, BlueprintReadWrite, Category = "Config", meta = (EditCondition = "WarpPointAnimProvider == EWarpPointAnimProvider::Bone"))
+	FName WarpPointAnimBoneName = NAME_None;
+
+	/** Whether to warp the translation component of the root motion */
+	UPROPERTY(EditAnywhere, BlueprintReadWrite, Category = "Config")
+	bool bWarpTranslation = true;
+
+	/** Whether to ignore the Z component of the translation. Z motion will remain untouched */
+	UPROPERTY(EditAnywhere, BlueprintReadWrite, Category = "Config", meta = (EditCondition = "bWarpTranslation"))
+	bool bIgnoreZAxis = true;
+
+	/** Whether to warp the rotation component of the root motion */
+	UPROPERTY(EditAnywhere, BlueprintReadWrite, Category = "Config")
+	bool bWarpRotation = true;
+
+	/** Whether rotation should be warp to match the rotation of the sync point or to face the sync point */
+	UPROPERTY(EditAnywhere, BlueprintReadWrite, Category = "Config", meta = (EditCondition = "bWarpRotation"))
+	EMotionWarpRotationType RotationType;
+
+	/**
+	 * Allow to modify how fast the rotation is warped.
+	 * e.g if the window duration is 2sec and this is 0.5, the target rotation will be reached in 1sec instead of 2sec
+	 */
+	UPROPERTY(EditAnywhere, BlueprintReadWrite, Category = "Config", meta = (EditCondition = "bWarpRotation"))
+		float WarpRotationTimeMultiplier = 1.f;
+
+	URootMotionModifierConfig_Warp(const FObjectInitializer& ObjectInitializer) : Super(ObjectInitializer) {}
+
+	UE_DEPRECATED(5.0, "Please use URootMotionModifier_Warp::AddRootMotionModifierSimpleWarp instead")
+	UFUNCTION(BlueprintCallable, Category = "Motion Warping", meta = (DeprecatedFunction))
+	static FRootMotionModifierHandle AddRootMotionModifierSimpleWarp(UMotionWarpingComponent* InMotionWarpingComp, const UAnimSequenceBase* InAnimation, float InStartTime, float InEndTime,
+			FName InSyncPointName, EWarpPointAnimProvider InWarpPointAnimProvider, FTransform InWarpPointAnimTransform, FName InWarpPointAnimBoneName,
+			bool bInWarpTranslation, bool bInIgnoreZAxis, bool bInWarpRotation, EMotionWarpRotationType InRotationType, float InWarpRotationTimeMultiplier = 1.f) {
+		return FRootMotionModifierHandle::InvalidHandle;
+	}
+};
+
+//@TODO: Temp for backward compatibility will be removed shortly
+UCLASS() class MOTIONWARPING_API UMotionModifier : public URootMotionModifier { GENERATED_BODY() };
+UCLASS() class MOTIONWARPING_API UMotionModifier_Warp : public URootMotionModifier_Warp { GENERATED_BODY() };
\ No newline at end of file
diff --git a/Engine/Plugins/Experimental/Animation/MotionWarping/Source/MotionWarping/Public/RootMotionModifier_AdjustmentBlendWarp.h b/Engine/Plugins/Experimental/Animation/MotionWarping/Source/MotionWarping/Public/RootMotionModifier_AdjustmentBlendWarp.h
--- a/Engine/Plugins/Experimental/Animation/MotionWarping/Source/MotionWarping/Public/RootMotionModifier_AdjustmentBlendWarp.h
+++ b/Engine/Plugins/Experimental/Animation/MotionWarping/Source/MotionWarping/Public/RootMotionModifier_AdjustmentBlendWarp.h
@@ -41,98 +41,31 @@ struct FMotionDeltaTrackContainer
 	}
 };
 
-USTRUCT()
-struct MOTIONWARPING_API FRootMotionModifier_AdjustmentBlendWarp : public FRootMotionModifier_Warp
-{
-	GENERATED_BODY()
-
-public:
-
-	UPROPERTY()
-	bool bWarpIKBones = false;
-
-	UPROPERTY()
-	TArray<FName> IKBones;
-
-	FRootMotionModifier_AdjustmentBlendWarp();
-	virtual ~FRootMotionModifier_AdjustmentBlendWarp() {}
-
-	virtual UScriptStruct* GetScriptStruct() const { return FRootMotionModifier_AdjustmentBlendWarp::StaticStruct(); }
-	virtual void OnTargetTransformChanged() override;
-	virtual FTransform ProcessRootMotion(UMotionWarpingComponent& OwnerComp, const FTransform& InRootMotion, float DeltaSeconds) override;
-
-	void GetIKBoneTransformAndAlpha(FName BoneName, FTransform& OutTransform, float& OutAlpha) const;
-
-protected:
-
-	UPROPERTY()
-	FTransform CachedMeshTransform;
-
-	UPROPERTY()
-	FTransform CachedMeshRelativeTransform;
-
-	UPROPERTY()
-	FTransform CachedRootMotion;
-
-	UPROPERTY()
-	FAnimSequenceTrackContainer Result;
-
-	UPROPERTY()
-	float ActualStartTime = 0.f;
-
-	void PrecomputeWarpedTracks();
-
-	FTransform ExtractWarpedRootMotion() const;
-
-	void ExtractBoneTransformAtTime(FTransform& OutTransform, const FName& BoneName, float Time) const;
-	void ExtractBoneTransformAtTime(FTransform& OutTransform, int32 TrackIndex, float Time) const;
-	void ExtractBoneTransformAtFrame(FTransform& OutTransform, int32 TrackIndex, int32 Frame) const;
-
-#if !(UE_BUILD_SHIPPING || UE_BUILD_TEST)
-	void DrawDebugWarpedTracks(float DrawDuration) const;
-#endif
-
-	static void ExtractMotionDeltaFromRange(const FBoneContainer& BoneContainer, const UAnimSequenceBase* Animation, float StartTime, float EndTime, float SampleRate, FMotionDeltaTrackContainer& OutMotionDeltaTracks);
-
-	static void AdjustmentBlendWarp(const FBoneContainer& BoneContainer, const FCSPose<FCompactPose>& AdditivePose, const FMotionDeltaTrackContainer& MotionDeltaTracks, FAnimSequenceTrackContainer& Output);
-};
-
-UCLASS(meta = (DisplayName = "Adjustment Blend Warp"))
-class MOTIONWARPING_API URootMotionModifierConfig_AdjustmentBlendWarp : public URootMotionModifierConfig_Warp
+//DEPRECATED will be removed shortly
+UCLASS()
+class MOTIONWARPING_API URootMotionModifierConfig_AdjustmentBlendWarp : public UObject
 {
 	GENERATED_BODY()
 
 public:
 
-	UPROPERTY(EditAnywhere, BlueprintReadWrite, Category = "Config")
-	bool bWarpIKBones;
-
-	UPROPERTY(EditAnywhere, BlueprintReadWrite, Category = "Config", meta = (EditCondition = "bWarpIKBones"))
-	TArray<FName> IKBones;
-
-	URootMotionModifierConfig_AdjustmentBlendWarp(const FObjectInitializer& ObjectInitializer)
-		: Super(ObjectInitializer) {}
-
-	virtual FRootMotionModifierHandle AddRootMotionModifierNew(UMotionWarpingComponent* MotionWarpingComp, const UAnimSequenceBase* Animation, float StartTime, float EndTime) const override
-	{
-		return URootMotionModifierConfig_AdjustmentBlendWarp::AddRootMotionModifierAdjustmentBlendWarp(MotionWarpingComp, Animation, StartTime, EndTime, 
-			SyncPointName, WarpPointAnimProvider, WarpPointAnimTransform, WarpPointAnimBoneName,
-			bWarpTranslation, bIgnoreZAxis, bWarpRotation, bWarpIKBones, IKBones);
-	}
+	URootMotionModifierConfig_AdjustmentBlendWarp(const FObjectInitializer& ObjectInitializer) : Super(ObjectInitializer) {}
 
-	UFUNCTION(BlueprintCallable, Category = "Motion Warping")
+	UE_DEPRECATED(5.0, "Please use UMotionModifier_AdjustmentBlendWarp::AddRootMotionModifierAdjustmentBlendWarp instead")
+	UFUNCTION(BlueprintCallable, Category = "Motion Warping", meta = (DeprecatedFunction))
 	static FRootMotionModifierHandle AddRootMotionModifierAdjustmentBlendWarp(UMotionWarpingComponent* InMotionWarpingComp, const UAnimSequenceBase* InAnimation, float InStartTime, float InEndTime, 
 		FName InSyncPointName, EWarpPointAnimProvider InWarpPointAnimProvider, FTransform InWarpPointAnimTransform, FName InWarpPointAnimBoneName,
 		bool bInWarpTranslation, bool bInIgnoreZAxis, bool bInWarpRotation, bool bInWarpIKBones, const TArray<FName>& InIKBones);
 
-	UFUNCTION(BlueprintPure, Category = "Motion Warping")
+	UE_DEPRECATED(5.0, "Please use UMotionModifier_AdjustmentBlendWarp::GetAdjustmentBlendIKBoneTransformAndAlpha instead")
+	UFUNCTION(BlueprintPure, Category = "Motion Warping", meta = (DeprecatedFunction))
 	static void GetIKBoneTransformAndAlpha(ACharacter* Character, FName BoneName, FTransform& OutTransform, float& OutAlpha);
 };
 
 //////////////////////////////////////////////
 
 UCLASS()
-class MOTIONWARPING_API UMotionModifier_AdjustmentBlendWarp : public UMotionModifier_Warp
+class MOTIONWARPING_API URootMotionModifier_AdjustmentBlendWarp : public URootMotionModifier_Warp
 {
 	GENERATED_BODY()
 
@@ -144,13 +77,21 @@ public:
 	UPROPERTY(EditAnywhere, BlueprintReadWrite, Category = "Config")
 	TArray<FName> IKBones;
 
-	UMotionModifier_AdjustmentBlendWarp(const FObjectInitializer& ObjectInitializer);
+	URootMotionModifier_AdjustmentBlendWarp(const FObjectInitializer& ObjectInitializer);
 
 	virtual void OnTargetTransformChanged() override;
 	virtual FTransform ProcessRootMotion(const FTransform& InRootMotion, float DeltaSeconds) override;
 
 	void GetIKBoneTransformAndAlpha(FName BoneName, FTransform& OutTransform, float& OutAlpha) const;
 
+	UFUNCTION(BlueprintCallable, Category = "Motion Warping")
+	static URootMotionModifier_AdjustmentBlendWarp* AddRootMotionModifierAdjustmentBlendWarp(UMotionWarpingComponent* InMotionWarpingComp, const UAnimSequenceBase* InAnimation, float InStartTime, float InEndTime,
+		FName InWarpTargetName, EWarpPointAnimProvider InWarpPointAnimProvider, FTransform InWarpPointAnimTransform, FName InWarpPointAnimBoneName,
+		bool bInWarpTranslation, bool bInIgnoreZAxis, bool bInWarpRotation, bool bInWarpIKBones, const TArray<FName>& InIKBones);
+
+	UFUNCTION(BlueprintPure, Category = "Motion Warping")
+	static void GetAdjustmentBlendIKBoneTransformAndAlpha(ACharacter* Character, FName BoneName, FTransform& OutTransform, float& OutAlpha);
+
 protected:
 
 	UPROPERTY()
diff --git a/Engine/Plugins/Experimental/Animation/MotionWarping/Source/MotionWarping/Public/RootMotionModifier_SkewWarp.h b/Engine/Plugins/Experimental/Animation/MotionWarping/Source/MotionWarping/Public/RootMotionModifier_SkewWarp.h
--- a/Engine/Plugins/Experimental/Animation/MotionWarping/Source/MotionWarping/Public/RootMotionModifier_SkewWarp.h
+++ b/Engine/Plugins/Experimental/Animation/MotionWarping/Source/MotionWarping/Public/RootMotionModifier_SkewWarp.h
@@ -6,38 +6,18 @@
 #include "RootMotionModifier.h"
 #include "RootMotionModifier_SkewWarp.generated.h"
 
-USTRUCT()
-struct MOTIONWARPING_API FRootMotionModifier_SkewWarp : public FRootMotionModifier_Warp
-{
-	GENERATED_BODY()
-
-public:
-
-	FRootMotionModifier_SkewWarp(){}
-	virtual ~FRootMotionModifier_SkewWarp() {}
-
-	virtual UScriptStruct* GetScriptStruct() const { return FRootMotionModifier_SkewWarp::StaticStruct(); }
-	virtual FTransform ProcessRootMotion(UMotionWarpingComponent& OwnerComp, const FTransform& InRootMotion, float DeltaSeconds) override;
-};
-
-UCLASS(meta = (DisplayName = "Skew Warp"))
-class MOTIONWARPING_API URootMotionModifierConfig_SkewWarp : public URootMotionModifierConfig_Warp
+//DEPRECATED will be removed shortly
+UCLASS()
+class MOTIONWARPING_API URootMotionModifierConfig_SkewWarp : public UObject
 {
 	GENERATED_BODY()
 
 public:
 
-	URootMotionModifierConfig_SkewWarp(const FObjectInitializer& ObjectInitializer)
-		: Super(ObjectInitializer) {}
+	URootMotionModifierConfig_SkewWarp(const FObjectInitializer& ObjectInitializer) : Super(ObjectInitializer) {}
 
-	virtual FRootMotionModifierHandle AddRootMotionModifierNew(UMotionWarpingComponent* MotionWarpingComp, const UAnimSequenceBase* Animation, float StartTime, float EndTime) const override
-	{
-		return URootMotionModifierConfig_SkewWarp::AddRootMotionModifierSkewWarp(MotionWarpingComp, Animation, StartTime, EndTime, 
-			SyncPointName, WarpPointAnimProvider, WarpPointAnimTransform, WarpPointAnimBoneName,
-			bWarpTranslation, bIgnoreZAxis, bWarpRotation, RotationType, WarpRotationTimeMultiplier);
-	}
-
-	UFUNCTION(BlueprintCallable, Category = "Motion Warping")
+	UE_DEPRECATED(5.0, "Please use UMotionModifier_SkewWarp::AddRootMotionModifierSkewWarp instead")
+	UFUNCTION(BlueprintCallable, Category = "Motion Warping", meta = (DeprecatedFunction))
 	static FRootMotionModifierHandle AddRootMotionModifierSkewWarp(UMotionWarpingComponent* InMotionWarpingComp, const UAnimSequenceBase* InAnimation, float InStartTime, float InEndTime, 
 		FName InSyncPointName, EWarpPointAnimProvider InWarpPointAnimProvider, FTransform InWarpPointAnimTransform, FName InWarpPointAnimBoneName,
 		bool bInWarpTranslation, bool bInIgnoreZAxis, bool bInWarpRotation, EMotionWarpRotationType InRotationType, float InWarpRotationTimeMultiplier = 1.f);
@@ -46,13 +26,18 @@ public:
 //////////////////////////////////////////////
 
 UCLASS()
-class MOTIONWARPING_API UMotionModifier_SkewWarp : public UMotionModifier_Warp
+class MOTIONWARPING_API URootMotionModifier_SkewWarp : public URootMotionModifier_Warp
 {
 	GENERATED_BODY()
 
 public:
 
-	UMotionModifier_SkewWarp(const FObjectInitializer& ObjectInitializer);
+	URootMotionModifier_SkewWarp(const FObjectInitializer& ObjectInitializer);
 
 	virtual FTransform ProcessRootMotion(const FTransform& InRootMotion, float DeltaSeconds) override;
+
+	UFUNCTION(BlueprintCallable, Category = "Motion Warping")
+	static URootMotionModifier_SkewWarp* AddRootMotionModifierSkewWarp(UMotionWarpingComponent* InMotionWarpingComp, const UAnimSequenceBase* InAnimation, float InStartTime, float InEndTime,
+		FName InWarpTargetName, EWarpPointAnimProvider InWarpPointAnimProvider, FTransform InWarpPointAnimTransform, FName InWarpPointAnimBoneName,
+		bool bInWarpTranslation, bool bInIgnoreZAxis, bool bInWarpRotation, EMotionWarpRotationType InRotationType, float InWarpRotationTimeMultiplier = 1.f);
 };
\ No newline at end of file
