Motion Warping - Removed deprecated types and updated notify events

diff --git a/Engine/Plugins/Experimental/Animation/MotionWarping/Source/MotionWarping/Private/AnimNotifyState_MotionWarping.cpp b/Engine/Plugins/Experimental/Animation/MotionWarping/Source/MotionWarping/Private/AnimNotifyState_MotionWarping.cpp
--- a/Engine/Plugins/Experimental/Animation/MotionWarping/Source/MotionWarping/Private/AnimNotifyState_MotionWarping.cpp
+++ b/Engine/Plugins/Experimental/Animation/MotionWarping/Source/MotionWarping/Private/AnimNotifyState_MotionWarping.cpp
@@ -9,24 +9,11 @@ UAnimNotifyState_MotionWarping::UAnimNotifyState_MotionWarping(const FObjectInit
 {
 }
 
-FRootMotionModifierHandle UAnimNotifyState_MotionWarping::OnBecomeRelevant(UMotionWarpingComponent* MotionWarpingComp, const UAnimSequenceBase* Animation, float StartTime, float EndTime) const
+void UAnimNotifyState_MotionWarping::OnBecomeRelevant(UMotionWarpingComponent* MotionWarpingComp, const UAnimSequenceBase* Animation, float StartTime, float EndTime) const
 {
-	FRootMotionModifierHandle Handle = AddRootMotionModifier(MotionWarpingComp, Animation, StartTime, EndTime);
+	URootMotionModifier* RootMotionModifierNew = AddRootMotionModifier(MotionWarpingComp, Animation, StartTime, EndTime);
 
-	//@TODO: Temp, to handle the case where the blueprint notify does not return the handle. Will be removed shortly
-	if (!Handle.IsValid())
-	{
-		if (MotionWarpingComp->GetModifiers().Num())
-		{
-			URootMotionModifier* Last = MotionWarpingComp->GetModifiers().Last();
-			if (Last && Last->Animation.Get() == Animation && Last->StartTime == StartTime && Last->EndTime == EndTime)
-			{
-				Handle = Last->GetHandle();
-			}
-		}
-	}
-
-	if (URootMotionModifier* RootMotionModifierNew = MotionWarpingComp->GetModifierByHandle(Handle))
+	if (RootMotionModifierNew)
 	{
 		if (!RootMotionModifierNew->OnActivateDelegate.IsBound())
 		{
@@ -42,44 +29,33 @@ FRootMotionModifierHandle UAnimNotifyState_MotionWarping::OnBecomeRelevant(UMoti
 		{
 			RootMotionModifierNew->OnDeactivateDelegate.BindDynamic(this, &UAnimNotifyState_MotionWarping::OnRootMotionModifierDeactivate);
 		}
-
-		return Handle;
 	}
-
-	return Handle;
 }
 
-FRootMotionModifierHandle UAnimNotifyState_MotionWarping::AddRootMotionModifier_Implementation(UMotionWarpingComponent* MotionWarpingComp, const UAnimSequenceBase* Animation, float StartTime, float EndTime) const
+URootMotionModifier* UAnimNotifyState_MotionWarping::AddRootMotionModifier_Implementation(UMotionWarpingComponent* MotionWarpingComp, const UAnimSequenceBase* Animation, float StartTime, float EndTime) const
 {
-	if (MotionWarpingComp)
+	if (MotionWarpingComp && RootMotionModifier)
 	{
-		if (RootMotionModifier)
-		{
-			URootMotionModifier* NewRootMotionModifier = MotionWarpingComp->AddModifierFromTemplate(RootMotionModifier, Animation, StartTime, EndTime);
-			if (NewRootMotionModifier)
-			{
-				return NewRootMotionModifier->GetHandle();
-			}
-		}
+		return MotionWarpingComp->AddModifierFromTemplate(RootMotionModifier, Animation, StartTime, EndTime);
 	}
 
-	return FRootMotionModifierHandle::InvalidHandle;
+	return nullptr;
 }
 
-void UAnimNotifyState_MotionWarping::OnRootMotionModifierActivate(UMotionWarpingComponent* MotionWarpingComp, const FRootMotionModifierHandle& Handle)
+void UAnimNotifyState_MotionWarping::OnRootMotionModifierActivate(UMotionWarpingComponent* MotionWarpingComp, URootMotionModifier* Modifier)
 {
 	// Notify blueprint
-	OnWarpBegin(MotionWarpingComp, Handle);
+	OnWarpBegin(MotionWarpingComp, Modifier);
 }
 
-void UAnimNotifyState_MotionWarping::OnRootMotionModifierUpdate(UMotionWarpingComponent* MotionWarpingComp, const FRootMotionModifierHandle& Handle)
+void UAnimNotifyState_MotionWarping::OnRootMotionModifierUpdate(UMotionWarpingComponent* MotionWarpingComp, URootMotionModifier* Modifier)
 {
 	// Notify blueprint
-	OnWarpUpdate(MotionWarpingComp, Handle);
+	OnWarpUpdate(MotionWarpingComp, Modifier);
 }
 
-void UAnimNotifyState_MotionWarping::OnRootMotionModifierDeactivate(UMotionWarpingComponent* MotionWarpingComp, const FRootMotionModifierHandle& Handle)
+void UAnimNotifyState_MotionWarping::OnRootMotionModifierDeactivate(UMotionWarpingComponent* MotionWarpingComp, URootMotionModifier* Modifier)
 {
 	// Notify blueprint
-	OnWarpEnd(MotionWarpingComp, Handle);
+	OnWarpEnd(MotionWarpingComp, Modifier);
 }
\ No newline at end of file
diff --git a/Engine/Plugins/Experimental/Animation/MotionWarping/Source/MotionWarping/Private/MotionWarpingComponent.cpp b/Engine/Plugins/Experimental/Animation/MotionWarping/Source/MotionWarping/Private/MotionWarpingComponent.cpp
--- a/Engine/Plugins/Experimental/Animation/MotionWarping/Source/MotionWarping/Private/MotionWarpingComponent.cpp
+++ b/Engine/Plugins/Experimental/Animation/MotionWarping/Source/MotionWarping/Private/MotionWarpingComponent.cpp
@@ -178,39 +178,18 @@ bool UMotionWarpingComponent::ContainsModifier(const UAnimSequenceBase* Animatio
 		});
 }
 
-FRootMotionModifierHandle UMotionWarpingComponent::AddModifier(URootMotionModifier* Modifier)
+int32 UMotionWarpingComponent::AddModifier(URootMotionModifier* Modifier)
 {
 	if (ensureAlways(Modifier))
 	{
-		//@TODO: Move this from here after update all modifiers
-		Modifier->Initialize(this);
-
-		Modifiers.Add(Modifier);
-
 		UE_LOG(LogMotionWarping, Verbose, TEXT("MotionWarping: RootMotionModifier added. NetMode: %d WorldTime: %f Char: %s Animation: %s [%f %f] [%f %f] Loc: %s Rot: %s"),
 			GetWorld()->GetNetMode(), GetWorld()->GetTimeSeconds(), *GetNameSafe(GetCharacterOwner()), *GetNameSafe(Modifier->Animation.Get()), Modifier->StartTime, Modifier->EndTime, Modifier->PreviousPosition, Modifier->CurrentPosition,
 			*GetCharacterOwner()->GetActorLocation().ToString(), *GetCharacterOwner()->GetActorRotation().ToCompactString());
 
-		return Modifier->GetHandle();
+		return Modifiers.Add(Modifier);
 	}
 
-	return FRootMotionModifierHandle::InvalidHandle;
-}
-
-URootMotionModifier* UMotionWarpingComponent::GetModifierByHandle(const FRootMotionModifierHandle& Handle) const
-{
-	if (Handle.IsValid())
-	{
-		for (URootMotionModifier* RootMotionModifier : Modifiers)
-		{
-			if (RootMotionModifier && RootMotionModifier->GetHandle() == Handle)
-			{
-				return RootMotionModifier;
-			}
-		}
-	}
-
-	return nullptr;
+	return INDEX_NONE;
 }
 
 void UMotionWarpingComponent::DisableAllRootMotionModifiers()
@@ -410,32 +389,6 @@ int32 UMotionWarpingComponent::RemoveSyncPoint(FName Name)
 	return SyncPoints.Remove(Name);
 }
 
-bool UMotionWarpingComponent::SetRootMotionModifierDelegatesByHandle(FRootMotionModifierHandle Handle, FOnRootMotionModifierDelegate OnActivate, FOnRootMotionModifierDelegate OnUpdate, FOnRootMotionModifierDelegate OnDeactivate)
-{
-	if (URootMotionModifier* RootMotionModifier = GetModifierByHandle(Handle))
-	{
-		RootMotionModifier->OnActivateDelegate = OnActivate;
-		RootMotionModifier->OnUpdateDelegate = OnUpdate;
-		RootMotionModifier->OnDeactivateDelegate = OnDeactivate;
-		return true;
-	}
-
-	return false;
-}
-
-bool UMotionWarpingComponent::GetAnimationAndTimeRangeForRootMotionModifierByHandle(FRootMotionModifierHandle Handle, UAnimSequenceBase*& OutAnimation, float& OutStartTime, float& OutEndTime) const
-{
-	if (URootMotionModifier* RootMotionModifier = GetModifierByHandle(Handle))
-	{
-		OutAnimation = const_cast<UAnimSequenceBase*>(RootMotionModifier->GetAnimation());
-		OutStartTime = RootMotionModifier->StartTime;
-		OutEndTime = RootMotionModifier->EndTime;
-		return true;
-	}
-
-	return false;
-}
-
 URootMotionModifier* UMotionWarpingComponent::AddModifierFromTemplate(URootMotionModifier* Template, const UAnimSequenceBase* Animation, float StartTime, float EndTime)
 {
 	if (ensureAlways(Template))
diff --git a/Engine/Plugins/Experimental/Animation/MotionWarping/Source/MotionWarping/Private/RootMotionModifier.cpp b/Engine/Plugins/Experimental/Animation/MotionWarping/Source/MotionWarping/Private/RootMotionModifier.cpp
--- a/Engine/Plugins/Experimental/Animation/MotionWarping/Source/MotionWarping/Private/RootMotionModifier.cpp
+++ b/Engine/Plugins/Experimental/Animation/MotionWarping/Source/MotionWarping/Private/RootMotionModifier.cpp
@@ -10,25 +10,14 @@
 #include "Animation/AnimInstance.h"
 #include "AnimNotifyState_MotionWarping.h"
 
-const FRootMotionModifierHandle FRootMotionModifierHandle::InvalidHandle;
-
-void FRootMotionModifierHandle::GenerateNewHandle()
-{
-	// Must be in C++ to avoid duplicate statics across execution units
-	static int32 GHandle = 1;
-	Handle = GHandle++;
-}
-
 URootMotionModifier::URootMotionModifier(const FObjectInitializer& ObjectInitializer)
 	: Super(ObjectInitializer)
 {
 }
 
-void URootMotionModifier::Initialize(UMotionWarpingComponent* OwnerComp)
+UMotionWarpingComponent* URootMotionModifier::GetOwnerComponent() const
 {
-	check(OwnerComp);
-	OwnerComponent = OwnerComp;
-	Handle.GenerateNewHandle();
+	return Cast<UMotionWarpingComponent>(GetOuter());
 }
 
 ACharacter* URootMotionModifier::GetCharacterOwner() const
@@ -87,7 +76,7 @@ void URootMotionModifier::Update()
 	{
 		if (UMotionWarpingComponent* OwnerComp = GetOwnerComponent())
 		{
-			OnUpdateDelegate.ExecuteIfBound(OwnerComp, Handle);
+			OnUpdateDelegate.ExecuteIfBound(OwnerComp, this);
 		}
 	}
 }
@@ -110,11 +99,11 @@ void URootMotionModifier::OnStateChanged(ERootMotionModifierState LastState)
 	{
 		if (LastState != ERootMotionModifierState::Active && State == ERootMotionModifierState::Active)
 		{
-			OnActivateDelegate.ExecuteIfBound(OwnerComp, Handle);
+			OnActivateDelegate.ExecuteIfBound(OwnerComp, this);
 		}
 		else if (LastState == ERootMotionModifierState::Active && (State == ERootMotionModifierState::Disabled || State == ERootMotionModifierState::MarkedForRemoval))
 		{
-			OnDeactivateDelegate.ExecuteIfBound(OwnerComp, Handle);
+			OnDeactivateDelegate.ExecuteIfBound(OwnerComp, this);
 		}
 	}
 }
diff --git a/Engine/Plugins/Experimental/Animation/MotionWarping/Source/MotionWarping/Private/RootMotionModifier_AdjustmentBlendWarp.cpp b/Engine/Plugins/Experimental/Animation/MotionWarping/Source/MotionWarping/Private/RootMotionModifier_AdjustmentBlendWarp.cpp
--- a/Engine/Plugins/Experimental/Animation/MotionWarping/Source/MotionWarping/Private/RootMotionModifier_AdjustmentBlendWarp.cpp
+++ b/Engine/Plugins/Experimental/Animation/MotionWarping/Source/MotionWarping/Private/RootMotionModifier_AdjustmentBlendWarp.cpp
@@ -14,42 +14,6 @@
 DECLARE_CYCLE_STAT(TEXT("MotionWarping PrecomputeWarpedTracks"), STAT_MotionWarping_PrecomputeWarpedTracks, STATGROUP_Anim);
 DECLARE_CYCLE_STAT(TEXT("MotionWarping ExtractMotionDelta"), STAT_MotionWarping_ExtractMotionDelta, STATGROUP_Anim);
 
-void URootMotionModifierConfig_AdjustmentBlendWarp::GetIKBoneTransformAndAlpha(ACharacter* Character, FName BoneName, FTransform& OutTransform, float& OutAlpha)
-{
-	URootMotionModifier_AdjustmentBlendWarp::GetAdjustmentBlendIKBoneTransformAndAlpha(Character, BoneName, OutTransform, OutAlpha);
-}
-
-// URootMotionModifierConfig_AdjustmentBlendWarp
-///////////////////////////////////////////////////////////////
-
-FRootMotionModifierHandle URootMotionModifierConfig_AdjustmentBlendWarp::AddRootMotionModifierAdjustmentBlendWarp(UMotionWarpingComponent* InMotionWarpingComp, const UAnimSequenceBase* InAnimation, float InStartTime, float InEndTime, 
-	FName InSyncPointName, EWarpPointAnimProvider InWarpPointAnimProvider, FTransform InWarpPointAnimTransform, FName InWarpPointAnimBoneName,
-	bool bInWarpTranslation, bool bInIgnoreZAxis, bool bInWarpRotation, bool bInWarpIKBones, const TArray<FName>& InIKBones)
-{
-	if (ensureAlways(InMotionWarpingComp))
-	{
-		URootMotionModifier_AdjustmentBlendWarp* NewModifier = NewObject<URootMotionModifier_AdjustmentBlendWarp>(InMotionWarpingComp);
-		NewModifier->Animation = InAnimation;
-		NewModifier->StartTime = InStartTime;
-		NewModifier->EndTime = InEndTime;
-		NewModifier->WarpTargetName = InSyncPointName;
-		NewModifier->WarpPointAnimProvider = InWarpPointAnimProvider;
-		NewModifier->WarpPointAnimTransform = InWarpPointAnimTransform;
-		NewModifier->WarpPointAnimBoneName = InWarpPointAnimBoneName;
-		NewModifier->bWarpTranslation = bInWarpTranslation;
-		NewModifier->bIgnoreZAxis = bInIgnoreZAxis;
-		NewModifier->bWarpRotation = bInWarpRotation;
-		NewModifier->bWarpIKBones = bInWarpIKBones;
-		NewModifier->IKBones = InIKBones;
-
-		return InMotionWarpingComp->AddModifier(NewModifier);
-	}
-
-	return FRootMotionModifierHandle::InvalidHandle;
-}
-
-/////////////////////////////////////////////////////////
-
 URootMotionModifier_AdjustmentBlendWarp::URootMotionModifier_AdjustmentBlendWarp(const FObjectInitializer& ObjectInitializer)
 	: Super(ObjectInitializer)
 {
diff --git a/Engine/Plugins/Experimental/Animation/MotionWarping/Source/MotionWarping/Private/RootMotionModifier_SkewWarp.cpp b/Engine/Plugins/Experimental/Animation/MotionWarping/Source/MotionWarping/Private/RootMotionModifier_SkewWarp.cpp
--- a/Engine/Plugins/Experimental/Animation/MotionWarping/Source/MotionWarping/Private/RootMotionModifier_SkewWarp.cpp
+++ b/Engine/Plugins/Experimental/Animation/MotionWarping/Source/MotionWarping/Private/RootMotionModifier_SkewWarp.cpp
@@ -7,36 +7,6 @@
 #include "Components/SkeletalMeshComponent.h"
 #include "MotionWarpingComponent.h"
 
-// URootMotionModifierConfig_SkewWarp
-///////////////////////////////////////////////////////////////
-
-FRootMotionModifierHandle URootMotionModifierConfig_SkewWarp::AddRootMotionModifierSkewWarp(UMotionWarpingComponent* InMotionWarpingComp, const UAnimSequenceBase* InAnimation, float InStartTime, float InEndTime, 
-	FName InSyncPointName, EWarpPointAnimProvider InWarpPointAnimProvider, FTransform InWarpPointAnimTransform, FName InWarpPointAnimBoneName,
-	bool bInWarpTranslation, bool bInIgnoreZAxis, bool bInWarpRotation, EMotionWarpRotationType InRotationType, float InWarpRotationTimeMultiplier)
-{
-	if (ensureAlways(InMotionWarpingComp))
-	{
-		URootMotionModifier_SkewWarp* NewModifier = NewObject<URootMotionModifier_SkewWarp>(InMotionWarpingComp);
-		NewModifier->Animation = InAnimation;
-		NewModifier->StartTime = InStartTime;
-		NewModifier->EndTime = InEndTime;
-		NewModifier->WarpTargetName = InSyncPointName;
-		NewModifier->WarpPointAnimProvider = InWarpPointAnimProvider;
-		NewModifier->WarpPointAnimTransform = InWarpPointAnimTransform;
-		NewModifier->WarpPointAnimBoneName = InWarpPointAnimBoneName;
-		NewModifier->bWarpTranslation = bInWarpTranslation;
-		NewModifier->bIgnoreZAxis = bInIgnoreZAxis;
-		NewModifier->bWarpRotation = bInWarpRotation;
-		NewModifier->RotationType = InRotationType;
-		NewModifier->WarpRotationTimeMultiplier = InWarpRotationTimeMultiplier;
-		return InMotionWarpingComp->AddModifier(NewModifier);
-	}
-
-	return FRootMotionModifierHandle::InvalidHandle;
-}
-
-///////////////////////////////////////////////////////////////
-
 URootMotionModifier_SkewWarp::URootMotionModifier_SkewWarp(const FObjectInitializer& ObjectInitializer)
 	: Super(ObjectInitializer)
 {
diff --git a/Engine/Plugins/Experimental/Animation/MotionWarping/Source/MotionWarping/Public/AnimNotifyState_MotionWarping.h b/Engine/Plugins/Experimental/Animation/MotionWarping/Source/MotionWarping/Public/AnimNotifyState_MotionWarping.h
--- a/Engine/Plugins/Experimental/Animation/MotionWarping/Source/MotionWarping/Public/AnimNotifyState_MotionWarping.h
+++ b/Engine/Plugins/Experimental/Animation/MotionWarping/Source/MotionWarping/Public/AnimNotifyState_MotionWarping.h
@@ -26,27 +26,27 @@ public:
 	UAnimNotifyState_MotionWarping(const FObjectInitializer& ObjectInitializer);
 
 	/** Called from the MotionWarpingComp when this notify becomes relevant. See: UMotionWarpingComponent::Update */
-	FRootMotionModifierHandle OnBecomeRelevant(UMotionWarpingComponent* MotionWarpingComp, const UAnimSequenceBase* Animation, float StartTime, float EndTime) const;
+	void OnBecomeRelevant(UMotionWarpingComponent* MotionWarpingComp, const UAnimSequenceBase* Animation, float StartTime, float EndTime) const;
 
 	/** Creates a root motion modifier from the config class defined in the notify */
 	UFUNCTION(BlueprintNativeEvent, Category = "Motion Warping")
-	FRootMotionModifierHandle AddRootMotionModifier(UMotionWarpingComponent* MotionWarpingComp, const UAnimSequenceBase* Animation, float StartTime, float EndTime) const;
+	URootMotionModifier* AddRootMotionModifier(UMotionWarpingComponent* MotionWarpingComp, const UAnimSequenceBase* Animation, float StartTime, float EndTime) const;
 
 	UFUNCTION()
-	void OnRootMotionModifierActivate(UMotionWarpingComponent* MotionWarpingComp, const FRootMotionModifierHandle& Handle);
+	void OnRootMotionModifierActivate(UMotionWarpingComponent* MotionWarpingComp, URootMotionModifier* Modifier);
 
 	UFUNCTION()
-	void OnRootMotionModifierUpdate(UMotionWarpingComponent* MotionWarpingComp, const FRootMotionModifierHandle& Handle);
+	void OnRootMotionModifierUpdate(UMotionWarpingComponent* MotionWarpingComp, URootMotionModifier* Modifier);
 
 	UFUNCTION()
-	void OnRootMotionModifierDeactivate(UMotionWarpingComponent* MotionWarpingComp, const FRootMotionModifierHandle& Handle);
+	void OnRootMotionModifierDeactivate(UMotionWarpingComponent* MotionWarpingComp, URootMotionModifier* Modifier);
 
 	UFUNCTION(BlueprintImplementableEvent, Category = "Motion Warping")
-	void OnWarpBegin(UMotionWarpingComponent* MotionWarpingComp, const FRootMotionModifierHandle& Handle) const;
+	void OnWarpBegin(UMotionWarpingComponent* MotionWarpingComp, URootMotionModifier* Modifier) const;
 
 	UFUNCTION(BlueprintImplementableEvent, Category = "Motion Warping")
-	void OnWarpUpdate(UMotionWarpingComponent* MotionWarpingComp, const FRootMotionModifierHandle& Handle) const;
+	void OnWarpUpdate(UMotionWarpingComponent* MotionWarpingComp, URootMotionModifier* Modifier) const;
 
 	UFUNCTION(BlueprintImplementableEvent, Category = "Motion Warping")
-	void OnWarpEnd(UMotionWarpingComponent* MotionWarpingComp, const FRootMotionModifierHandle& Handle) const;
+	void OnWarpEnd(UMotionWarpingComponent* MotionWarpingComp, URootMotionModifier* Modifier) const;
 };
\ No newline at end of file
diff --git a/Engine/Plugins/Experimental/Animation/MotionWarping/Source/MotionWarping/Public/MotionWarpingComponent.h b/Engine/Plugins/Experimental/Animation/MotionWarping/Source/MotionWarping/Public/MotionWarpingComponent.h
--- a/Engine/Plugins/Experimental/Animation/MotionWarping/Source/MotionWarping/Public/MotionWarpingComponent.h
+++ b/Engine/Plugins/Experimental/Animation/MotionWarping/Source/MotionWarping/Public/MotionWarpingComponent.h
@@ -119,14 +119,11 @@ public:
 	bool ContainsModifier(const UAnimSequenceBase* Animation, float StartTime, float EndTime) const;
 
 	/** Add a new modifier */
-	FRootMotionModifierHandle AddModifier(URootMotionModifier* Modifier);
+	int32 AddModifier(URootMotionModifier* Modifier);
 
 	/** Mark all the modifiers as Disable */
 	void DisableAllRootMotionModifiers();
 
-	/** Finds a RootMotionModifier that matches the supplied Handle */
-	URootMotionModifier* GetModifierByHandle(const FRootMotionModifierHandle& Handle) const;
-
 	URootMotionModifier* AddModifierFromTemplate(URootMotionModifier* Template, const UAnimSequenceBase* Animation, float StartTime, float EndTime);
 
 protected:
@@ -152,12 +149,4 @@ protected:
 	FTransform ProcessRootMotionPreConvertToWorld(const FTransform& InRootMotion, class UCharacterMovementComponent* CharacterMovementComponent, float DeltaSeconds);
 	
 	FTransform ProcessRootMotionPostConvertToWorld(const FTransform& InRootMotion, class UCharacterMovementComponent* CharacterMovementComponent, float DeltaSeconds);
-
-public:
-
-	UFUNCTION(BlueprintCallable, Category = "Motion Warping")
-	bool SetRootMotionModifierDelegatesByHandle(FRootMotionModifierHandle Handle, FOnRootMotionModifierDelegate OnActivate, FOnRootMotionModifierDelegate OnUpdate, FOnRootMotionModifierDelegate OnDeactivate);
-
-	UFUNCTION(BlueprintCallable, Category = "Motion Warping")
-	bool GetAnimationAndTimeRangeForRootMotionModifierByHandle(FRootMotionModifierHandle Handle, UAnimSequenceBase*& OutAnimation, float& OutStartTime, float& OutEndTime) const;
 };
\ No newline at end of file
diff --git a/Engine/Plugins/Experimental/Animation/MotionWarping/Source/MotionWarping/Public/RootMotionModifier.h b/Engine/Plugins/Experimental/Animation/MotionWarping/Source/MotionWarping/Public/RootMotionModifier.h
--- a/Engine/Plugins/Experimental/Animation/MotionWarping/Source/MotionWarping/Public/RootMotionModifier.h
+++ b/Engine/Plugins/Experimental/Animation/MotionWarping/Source/MotionWarping/Public/RootMotionModifier.h
@@ -8,6 +8,7 @@
 
 class UMotionWarpingComponent;
 class UAnimNotifyState_MotionWarping;
+class URootMotionModifier;
 
 /** The possible states of a Root Motion Modifier */
 UENUM(BlueprintType)
@@ -26,35 +27,7 @@ enum class ERootMotionModifierState : uint8
 	Disabled
 };
 
-/** Handle to identify a RootMotionModifier */
-USTRUCT(BlueprintType)
-struct MOTIONWARPING_API FRootMotionModifierHandle
-{
-	GENERATED_BODY()
-
-	FRootMotionModifierHandle() : Handle(INDEX_NONE) {}
-
-	void GenerateNewHandle();
-	
-	bool IsValid() const { return Handle != INDEX_NONE; }
-	void Invalidate() { Handle = INDEX_NONE; }
-	
-	bool operator==(const FRootMotionModifierHandle& Other) const { return Handle == Other.Handle; }
-	bool operator!=(const FRootMotionModifierHandle& Other) const { return Handle != Other.Handle; }
-	
-	friend uint32 GetTypeHash(const FRootMotionModifierHandle& InHandle) { return InHandle.Handle; }
-
-	FString ToString() const { return IsValid() ? FString::FromInt(Handle) : TEXT("Invalid"); }
-
-	static const FRootMotionModifierHandle InvalidHandle;
-
-private:
-
-	UPROPERTY()
-	int32 Handle;
-};
-
-DECLARE_DYNAMIC_DELEGATE_TwoParams(FOnRootMotionModifierDelegate, UMotionWarpingComponent*, MotionWarpingComp, const FRootMotionModifierHandle&, Handle);
+DECLARE_DYNAMIC_DELEGATE_TwoParams(FOnRootMotionModifierDelegate, UMotionWarpingComponent*, MotionWarpingComp, URootMotionModifier*, RootMotionModifier);
 
 /** Represents a point of alignment in the world */
 USTRUCT(BlueprintType, meta = (HasNativeMake = "MotionWarping.MotionWarpingUtilities.MakeMotionWarpingSyncPoint", HasNativeBreak = "MotionWarping.MotionWarpingUtilities.BreakMotionWarpingSyncPoint"))
@@ -170,9 +143,6 @@ public:
 
 	URootMotionModifier(const FObjectInitializer& ObjectInitializer);
 
-	/** Initialize this modifier */
-	void Initialize(UMotionWarpingComponent* OwnerComp);
-
 	/** Called when the state of the modifier changes */
 	virtual void OnStateChanged(ERootMotionModifierState LastState);
 
@@ -182,12 +152,8 @@ public:
 	/** Returns the state of the modifier */
 	FORCEINLINE ERootMotionModifierState GetState() const { return State; }
 
-	/** Returns the handle to identify this modifier */
-	UFUNCTION(BlueprintPure, Category = "Motion Warping")
-	FRootMotionModifierHandle GetHandle() const { return Handle; }
-
 	/** Returns a pointer to the component that owns this modifier */
-	FORCEINLINE UMotionWarpingComponent* GetOwnerComponent() const { return OwnerComponent.Get(); }
+	UMotionWarpingComponent* GetOwnerComponent() const;
 
 	/** Returns a pointer to the character that owns the component that owns this modifier */
 	class ACharacter* GetCharacterOwner() const;
@@ -201,14 +167,6 @@ private:
 
 	friend UMotionWarpingComponent;
 
-	/** Pointer back to the Component that owns this modifier */
-	UPROPERTY()
-	TWeakObjectPtr<UMotionWarpingComponent> OwnerComponent = nullptr;
-
-	/** Handle to identify this modifier */
-	UPROPERTY()
-	FRootMotionModifierHandle Handle;
-
 	/** Current state */
 	UPROPERTY()
 	ERootMotionModifierState State = ERootMotionModifierState::Waiting;
@@ -315,82 +273,4 @@ public:
 		FinalRootMotion.ScaleTranslation(Scale);
 		return FinalRootMotion;
 	}
-};
-
-/** Blueprint wrapper around the config properties of a root motion modifier */
-UCLASS(abstract, BlueprintType, EditInlineNew)
-class MOTIONWARPING_API URootMotionModifierConfig : public UObject
-{
-	GENERATED_BODY()
-
-public:
-
-	URootMotionModifierConfig(const FObjectInitializer& ObjectInitializer)
-		: Super(ObjectInitializer) {}
-
-	/** Adds a RootMotionModifier of the type this object represents */
-	virtual FRootMotionModifierHandle AddRootMotionModifierNew(UMotionWarpingComponent* InMotionWarpingComp, const UAnimSequenceBase* Animation, float StartTime, float EndTime) const { return FRootMotionModifierHandle::InvalidHandle; }
-
-	//DEPRECATED will be removed shortly
-	virtual void AddRootMotionModifier(UMotionWarpingComponent* InMotionWarpingComp, const UAnimSequenceBase* Animation, float StartTime, float EndTime) const {}
-};
-
-UCLASS(meta = (DisplayName = "Simple Warp"))
-class MOTIONWARPING_API URootMotionModifierConfig_Warp : public URootMotionModifierConfig
-{
-	GENERATED_BODY()
-
-public:
-
-	/** Name used to find the warp point for this modifier */
-	UPROPERTY(EditAnywhere, BlueprintReadWrite, Category = "Config", meta = (DisplayName = "Warp Point Name"))
-	FName SyncPointName = NAME_None;  //@TODO: Rename to WarpPointName
-
-	UPROPERTY(EditAnywhere, BlueprintReadWrite, Category = "Config")
-	EWarpPointAnimProvider WarpPointAnimProvider = EWarpPointAnimProvider::None;
-
-	//@TODO: Hide from the UI when Provider != Static
-	UPROPERTY(EditAnywhere, BlueprintReadWrite, Category = "Config", meta = (EditCondition = "WarpPointAnimProvider == EWarpPointAnimProvider::Static"))
-	FTransform WarpPointAnimTransform = FTransform::Identity;
-
-	//@TODO: Hide from the UI when Provider != Bone
-	UPROPERTY(EditAnywhere, BlueprintReadWrite, Category = "Config", meta = (EditCondition = "WarpPointAnimProvider == EWarpPointAnimProvider::Bone"))
-	FName WarpPointAnimBoneName = NAME_None;
-
-	/** Whether to warp the translation component of the root motion */
-	UPROPERTY(EditAnywhere, BlueprintReadWrite, Category = "Config")
-	bool bWarpTranslation = true;
-
-	/** Whether to ignore the Z component of the translation. Z motion will remain untouched */
-	UPROPERTY(EditAnywhere, BlueprintReadWrite, Category = "Config", meta = (EditCondition = "bWarpTranslation"))
-	bool bIgnoreZAxis = true;
-
-	/** Whether to warp the rotation component of the root motion */
-	UPROPERTY(EditAnywhere, BlueprintReadWrite, Category = "Config")
-	bool bWarpRotation = true;
-
-	/** Whether rotation should be warp to match the rotation of the sync point or to face the sync point */
-	UPROPERTY(EditAnywhere, BlueprintReadWrite, Category = "Config", meta = (EditCondition = "bWarpRotation"))
-	EMotionWarpRotationType RotationType;
-
-	/**
-	 * Allow to modify how fast the rotation is warped.
-	 * e.g if the window duration is 2sec and this is 0.5, the target rotation will be reached in 1sec instead of 2sec
-	 */
-	UPROPERTY(EditAnywhere, BlueprintReadWrite, Category = "Config", meta = (EditCondition = "bWarpRotation"))
-		float WarpRotationTimeMultiplier = 1.f;
-
-	URootMotionModifierConfig_Warp(const FObjectInitializer& ObjectInitializer) : Super(ObjectInitializer) {}
-
-	UE_DEPRECATED(5.0, "Please use URootMotionModifier_Warp::AddRootMotionModifierSimpleWarp instead")
-	UFUNCTION(BlueprintCallable, Category = "Motion Warping", meta = (DeprecatedFunction))
-	static FRootMotionModifierHandle AddRootMotionModifierSimpleWarp(UMotionWarpingComponent* InMotionWarpingComp, const UAnimSequenceBase* InAnimation, float InStartTime, float InEndTime,
-			FName InSyncPointName, EWarpPointAnimProvider InWarpPointAnimProvider, FTransform InWarpPointAnimTransform, FName InWarpPointAnimBoneName,
-			bool bInWarpTranslation, bool bInIgnoreZAxis, bool bInWarpRotation, EMotionWarpRotationType InRotationType, float InWarpRotationTimeMultiplier = 1.f) {
-		return FRootMotionModifierHandle::InvalidHandle;
-	}
-};
-
-//@TODO: Temp for backward compatibility will be removed shortly
-UCLASS() class MOTIONWARPING_API UMotionModifier : public URootMotionModifier { GENERATED_BODY() };
-UCLASS() class MOTIONWARPING_API UMotionModifier_Warp : public URootMotionModifier_Warp { GENERATED_BODY() };
\ No newline at end of file
+};
\ No newline at end of file
diff --git a/Engine/Plugins/Experimental/Animation/MotionWarping/Source/MotionWarping/Public/RootMotionModifier_AdjustmentBlendWarp.h b/Engine/Plugins/Experimental/Animation/MotionWarping/Source/MotionWarping/Public/RootMotionModifier_AdjustmentBlendWarp.h
--- a/Engine/Plugins/Experimental/Animation/MotionWarping/Source/MotionWarping/Public/RootMotionModifier_AdjustmentBlendWarp.h
+++ b/Engine/Plugins/Experimental/Animation/MotionWarping/Source/MotionWarping/Public/RootMotionModifier_AdjustmentBlendWarp.h
@@ -41,29 +41,6 @@ struct FMotionDeltaTrackContainer
 	}
 };
 
-//DEPRECATED will be removed shortly
-UCLASS()
-class MOTIONWARPING_API URootMotionModifierConfig_AdjustmentBlendWarp : public UObject
-{
-	GENERATED_BODY()
-
-public:
-
-	URootMotionModifierConfig_AdjustmentBlendWarp(const FObjectInitializer& ObjectInitializer) : Super(ObjectInitializer) {}
-
-	UE_DEPRECATED(5.0, "Please use UMotionModifier_AdjustmentBlendWarp::AddRootMotionModifierAdjustmentBlendWarp instead")
-	UFUNCTION(BlueprintCallable, Category = "Motion Warping", meta = (DeprecatedFunction))
-	static FRootMotionModifierHandle AddRootMotionModifierAdjustmentBlendWarp(UMotionWarpingComponent* InMotionWarpingComp, const UAnimSequenceBase* InAnimation, float InStartTime, float InEndTime, 
-		FName InSyncPointName, EWarpPointAnimProvider InWarpPointAnimProvider, FTransform InWarpPointAnimTransform, FName InWarpPointAnimBoneName,
-		bool bInWarpTranslation, bool bInIgnoreZAxis, bool bInWarpRotation, bool bInWarpIKBones, const TArray<FName>& InIKBones);
-
-	UE_DEPRECATED(5.0, "Please use UMotionModifier_AdjustmentBlendWarp::GetAdjustmentBlendIKBoneTransformAndAlpha instead")
-	UFUNCTION(BlueprintPure, Category = "Motion Warping", meta = (DeprecatedFunction))
-	static void GetIKBoneTransformAndAlpha(ACharacter* Character, FName BoneName, FTransform& OutTransform, float& OutAlpha);
-};
-
-//////////////////////////////////////////////
-
 UCLASS()
 class MOTIONWARPING_API URootMotionModifier_AdjustmentBlendWarp : public URootMotionModifier_Warp
 {
diff --git a/Engine/Plugins/Experimental/Animation/MotionWarping/Source/MotionWarping/Public/RootMotionModifier_SkewWarp.h b/Engine/Plugins/Experimental/Animation/MotionWarping/Source/MotionWarping/Public/RootMotionModifier_SkewWarp.h
--- a/Engine/Plugins/Experimental/Animation/MotionWarping/Source/MotionWarping/Public/RootMotionModifier_SkewWarp.h
+++ b/Engine/Plugins/Experimental/Animation/MotionWarping/Source/MotionWarping/Public/RootMotionModifier_SkewWarp.h
@@ -6,25 +6,6 @@
 #include "RootMotionModifier.h"
 #include "RootMotionModifier_SkewWarp.generated.h"
 
-//DEPRECATED will be removed shortly
-UCLASS()
-class MOTIONWARPING_API URootMotionModifierConfig_SkewWarp : public UObject
-{
-	GENERATED_BODY()
-
-public:
-
-	URootMotionModifierConfig_SkewWarp(const FObjectInitializer& ObjectInitializer) : Super(ObjectInitializer) {}
-
-	UE_DEPRECATED(5.0, "Please use UMotionModifier_SkewWarp::AddRootMotionModifierSkewWarp instead")
-	UFUNCTION(BlueprintCallable, Category = "Motion Warping", meta = (DeprecatedFunction))
-	static FRootMotionModifierHandle AddRootMotionModifierSkewWarp(UMotionWarpingComponent* InMotionWarpingComp, const UAnimSequenceBase* InAnimation, float InStartTime, float InEndTime, 
-		FName InSyncPointName, EWarpPointAnimProvider InWarpPointAnimProvider, FTransform InWarpPointAnimTransform, FName InWarpPointAnimBoneName,
-		bool bInWarpTranslation, bool bInIgnoreZAxis, bool bInWarpRotation, EMotionWarpRotationType InRotationType, float InWarpRotationTimeMultiplier = 1.f);
-};
-
-//////////////////////////////////////////////
-
 UCLASS()
 class MOTIONWARPING_API URootMotionModifier_SkewWarp : public URootMotionModifier_Warp
 {
