Motion Warping:
Fix warp point offset calculation.
BoneContainer to extract warp point from the animation is now created with the correct skeleton to account for retargeting.
Hide animation and time ranges from the notify panel since those should not be exposed there.

diff --git a/Engine/Plugins/Experimental/Animation/MotionWarping/Source/MotionWarping/Private/RootMotionModifier.cpp b/Engine/Plugins/Experimental/Animation/MotionWarping/Source/MotionWarping/Private/RootMotionModifier.cpp
--- a/Engine/Plugins/Experimental/Animation/MotionWarping/Source/MotionWarping/Private/RootMotionModifier.cpp
+++ b/Engine/Plugins/Experimental/Animation/MotionWarping/Source/MotionWarping/Private/RootMotionModifier.cpp
@@ -171,7 +171,7 @@ void URootMotionModifier_Warp::Update()
 									TArray<FBoneIndexType> RequiredBoneIndexArray = { 0, (FBoneIndexType)BoneIndex };
 									FullBoneContainer.GetReferenceSkeleton().EnsureParentsExistAndSort(RequiredBoneIndexArray);
 
-									FBoneContainer LimitedBoneContainer(RequiredBoneIndexArray, FCurveEvaluationOption(false), *Animation->GetSkeleton());
+									FBoneContainer LimitedBoneContainer(RequiredBoneIndexArray, FCurveEvaluationOption(false), *FullBoneContainer.GetAsset());
 
 									FCSPose<FCompactPose> Pose;
 									UMotionWarpingUtilities::ExtractComponentSpacePose(GetAnimation(), LimitedBoneContainer, EndTime, false, Pose);
@@ -184,7 +184,7 @@ void URootMotionModifier_Warp::Update()
 					}
 
 					// Calculate and cache the offset between the root and the warp point in the animation
-					CachedOffsetFromWarpPoint = WarpPointTransform.GetRelativeTransform(RootTransform);
+					CachedOffsetFromWarpPoint = RootTransform.GetRelativeTransform(WarpPointTransform);
 				}
 			}
 
diff --git a/Engine/Plugins/Experimental/Animation/MotionWarping/Source/MotionWarping/Public/RootMotionModifier.h b/Engine/Plugins/Experimental/Animation/MotionWarping/Source/MotionWarping/Public/RootMotionModifier.h
--- a/Engine/Plugins/Experimental/Animation/MotionWarping/Source/MotionWarping/Public/RootMotionModifier.h
+++ b/Engine/Plugins/Experimental/Animation/MotionWarping/Source/MotionWarping/Public/RootMotionModifier.h
@@ -102,15 +102,15 @@ class MOTIONWARPING_API URootMotionModifier : public UObject
 public:
 
 	/** Source of the root motion we are warping */
-	UPROPERTY(EditAnywhere, BlueprintReadWrite, Category = "Config", meta = (ExposeOnSpawn))
+	UPROPERTY()
 	TWeakObjectPtr<const UAnimSequenceBase> Animation = nullptr;
 
 	/** Start time of the warping window */
-	UPROPERTY(EditAnywhere, BlueprintReadWrite, Category = "Config", meta = (ExposeOnSpawn))
+	UPROPERTY()
 	float StartTime = 0.f;
 
 	/** End time of the warping window */
-	UPROPERTY(EditAnywhere, BlueprintReadWrite, Category = "Config", meta = (ExposeOnSpawn))
+	UPROPERTY()
 	float EndTime = 0.f;
 
 	/** Previous playback time of the animation */
