Fixed motion warping windows in vaulting example + re-added DisplayName to each root motion modifier

diff --git a/Engine/Plugins/Experimental/Animation/MotionWarping/Source/MotionWarping/Public/RootMotionModifier.h b/Engine/Plugins/Experimental/Animation/MotionWarping/Source/MotionWarping/Public/RootMotionModifier.h
--- a/Engine/Plugins/Experimental/Animation/MotionWarping/Source/MotionWarping/Public/RootMotionModifier.h
+++ b/Engine/Plugins/Experimental/Animation/MotionWarping/Source/MotionWarping/Public/RootMotionModifier.h
@@ -172,7 +172,7 @@ private:
 	ERootMotionModifierState State = ERootMotionModifierState::Waiting;
 };
 
-UCLASS()
+UCLASS(meta = (DisplayName = "Simple Warp"))
 class MOTIONWARPING_API URootMotionModifier_Warp : public URootMotionModifier
 {
 	GENERATED_BODY()
@@ -254,7 +254,7 @@ private:
 	TOptional<FTransform> CachedOffsetFromWarpPoint;
 };
 
-UCLASS()
+UCLASS(meta = (DisplayName = "Scale"))
 class MOTIONWARPING_API URootMotionModifier_Scale : public URootMotionModifier
 {
 	GENERATED_BODY()
diff --git a/Engine/Plugins/Experimental/Animation/MotionWarping/Source/MotionWarping/Public/RootMotionModifier_AdjustmentBlendWarp.h b/Engine/Plugins/Experimental/Animation/MotionWarping/Source/MotionWarping/Public/RootMotionModifier_AdjustmentBlendWarp.h
--- a/Engine/Plugins/Experimental/Animation/MotionWarping/Source/MotionWarping/Public/RootMotionModifier_AdjustmentBlendWarp.h
+++ b/Engine/Plugins/Experimental/Animation/MotionWarping/Source/MotionWarping/Public/RootMotionModifier_AdjustmentBlendWarp.h
@@ -41,7 +41,7 @@ struct FMotionDeltaTrackContainer
 	}
 };
 
-UCLASS()
+UCLASS(meta = (DisplayName = "Adjustment Blend Warp"))
 class MOTIONWARPING_API URootMotionModifier_AdjustmentBlendWarp : public URootMotionModifier_Warp
 {
 	GENERATED_BODY()
diff --git a/Engine/Plugins/Experimental/Animation/MotionWarping/Source/MotionWarping/Public/RootMotionModifier_SkewWarp.h b/Engine/Plugins/Experimental/Animation/MotionWarping/Source/MotionWarping/Public/RootMotionModifier_SkewWarp.h
--- a/Engine/Plugins/Experimental/Animation/MotionWarping/Source/MotionWarping/Public/RootMotionModifier_SkewWarp.h
+++ b/Engine/Plugins/Experimental/Animation/MotionWarping/Source/MotionWarping/Public/RootMotionModifier_SkewWarp.h
@@ -6,7 +6,7 @@
 #include "RootMotionModifier.h"
 #include "RootMotionModifier_SkewWarp.generated.h"
 
-UCLASS()
+UCLASS(meta = (DisplayName = "Skew Warp"))
 class MOTIONWARPING_API URootMotionModifier_SkewWarp : public URootMotionModifier_Warp
 {
 	GENERATED_BODY()
