Motion Warping bug fixes:
Fixed bug where character would miss the target if capsule collision is fully disable.
Fixed bug where skew warp would miss the target if the mesh is not perfectly aligned with the bottom of the capsule in the character blueprint.
Fixed bug where simple warp would stomp Z translation when ignoring Z axis.
Minor change to PrintLog to use motion warping category and log out Z deltas

diff --git a/Engine/Plugins/Experimental/Animation/MotionWarping/Source/MotionWarping/Private/RootMotionModifier.cpp b/Engine/Plugins/Experimental/Animation/MotionWarping/Source/MotionWarping/Private/RootMotionModifier.cpp
--- a/Engine/Plugins/Experimental/Animation/MotionWarping/Source/MotionWarping/Private/RootMotionModifier.cpp
+++ b/Engine/Plugins/Experimental/Animation/MotionWarping/Source/MotionWarping/Private/RootMotionModifier.cpp
@@ -3,6 +3,7 @@
 #include "RootMotionModifier.h"
 #include "GameFramework/Character.h"
 #include "GameFramework/CharacterMovementComponent.h"
+#include "Components/CapsuleComponent.h"
 #include "Components/SkeletalMeshComponent.h"
 #include "Animation/AnimMontage.h"
 #include "MotionWarpingComponent.h"
@@ -224,7 +225,7 @@ FTransform URootMotionModifier_Warp::ProcessRootMotion(const FTransform& InRootM
 		const float HorizontalDelta = RootMotionDelta.GetTranslation().Size2D();
 		const float HorizontalTarget = FVector::Dist2D(CharacterTransform.GetLocation(), GetTargetLocation());
 		const float HorizontalOriginal = RootMotionTotal.GetTranslation().Size2D();
-		const float HorizontalTranslationWarped = HorizontalOriginal != 0.f ? ((HorizontalDelta * HorizontalTarget) / HorizontalOriginal) : 0.f;
+		const float HorizontalTranslationWarped = !FMath::IsNearlyZero(HorizontalOriginal) ? ((HorizontalDelta * HorizontalTarget) / HorizontalOriginal) : 0.f;
 
 		if (bInLocalSpace)
 		{
@@ -239,14 +240,19 @@ FTransform URootMotionModifier_Warp::ProcessRootMotion(const FTransform& InRootM
 
 		if (!bIgnoreZAxis)
 		{
-			const FVector CapsuleBottomLocation = (CharacterOwner->GetActorLocation() - FVector::UpVector * CharacterOwner->GetSimpleCollisionHalfHeight());
+			const float CapsuleHalfHeight = CharacterOwner->GetCapsuleComponent()->GetScaledCapsuleHalfHeight();
+			const FVector CapsuleBottomLocation = (CharacterOwner->GetActorLocation() - FVector(0.f, 0.f, CapsuleHalfHeight));
 			const float VerticalDelta = RootMotionDelta.GetTranslation().Z;
 			const float VerticalTarget = GetTargetLocation().Z - CapsuleBottomLocation.Z;
 			const float VerticalOriginal = RootMotionTotal.GetTranslation().Z;
-			const float VerticalTranslationWarped = VerticalOriginal != 0.f ? ((VerticalDelta * VerticalTarget) / VerticalOriginal) : 0.f;
+			const float VerticalTranslationWarped = !FMath::IsNearlyZero(VerticalOriginal) ? ((VerticalDelta * VerticalTarget) / VerticalOriginal) : 0.f;
 
 			DeltaTranslation.Z = VerticalTranslationWarped;
 		}
+		else
+		{
+			DeltaTranslation.Z = InRootMotion.GetTranslation().Z;
+		}
 
 		FinalRootMotion.SetTranslation(DeltaTranslation);
 	}
@@ -262,7 +268,7 @@ FTransform URootMotionModifier_Warp::ProcessRootMotion(const FTransform& InRootM
 	const int32 DebugLevel = FMotionWarpingCVars::CVarMotionWarpingDebug.GetValueOnGameThread();
 	if (DebugLevel > 0)
 	{
-		PrintLog(TEXT("FRootMotionModifier_Simple"), InRootMotion, FinalRootMotion);
+		PrintLog(TEXT("SimpleWarp"), InRootMotion, FinalRootMotion);
 
 		if (DebugLevel >= 2)
 		{
@@ -320,7 +326,8 @@ void URootMotionModifier_Warp::PrintLog(const FString& Name, const FTransform& O
 {
 	if (const ACharacter* CharacterOwner = GetCharacterOwner())
 	{
-		const FVector CurrentLocation = (CharacterOwner->GetActorLocation() - FVector::UpVector * CharacterOwner->GetSimpleCollisionHalfHeight());
+		const float CapsuleHalfHeight = CharacterOwner->GetCapsuleComponent()->GetScaledCapsuleHalfHeight();
+		const FVector CurrentLocation = (CharacterOwner->GetActorLocation() - FVector(0.f, 0.f, CapsuleHalfHeight));
 		const FVector CurrentToTarget = (GetTargetLocation() - CurrentLocation).GetSafeNormal2D();
 		const FVector FutureLocation = CurrentLocation + (bInLocalSpace ? (CharacterOwner->GetMesh()->ConvertLocalRootMotionToWorld(WarpedRootMotion)).GetTranslation() : WarpedRootMotion.GetTranslation());
 		const FRotator CurrentRotation = CharacterOwner->GetActorRotation();
@@ -332,8 +339,9 @@ void URootMotionModifier_Warp::PrintLog(const FString& Name, const FTransform& O
 		const float Speed = WarpedRootMotion.GetTranslation().Size() / DeltaSeconds;
 		const float EndTimeOffset = CurrentPosition - EndTime;
 
-		UE_LOG(LogMotionWarping, Log, TEXT("MotionWarping: %s. NetMode: %d Char: %s Anim: %s Win: [%f %f][%f %f] DT: %f WT: %f EndTimeOffset: %f Dist2D: %f FutureDist2D: %f Dot: %f OriginalMotionDelta: %s (%f) FinalMotionDelta: %s (%f) Speed: %f Loc: %s FutureLoc: %s Rot: %s FutureRot: %s"),
-			*Name, (int32)CharacterOwner->GetWorld()->GetNetMode(), *GetNameSafe(CharacterOwner), *GetNameSafe(Animation.Get()), StartTime, EndTime, PreviousPosition, CurrentPosition, DeltaSeconds, CharacterOwner->GetWorld()->GetTimeSeconds(), EndTimeOffset, CurrentDist2D, FutureDist2D, Dot,
+		UE_LOG(LogMotionWarping, Log, TEXT("%s NetMode: %d Char: %s Anim: %s Win: [%f %f][%f %f] DT: %f WT: %f ETOffset: %f Dist2D: %f Z: %f FDist2D: %f FZ: %f Dot: %f Delta: %s (%f) FDelta: %s (%f) Speed: %f Loc: %s FLoc: %s Rot: %s FRot: %s"),
+			*Name, (int32)CharacterOwner->GetWorld()->GetNetMode(), *GetNameSafe(CharacterOwner), *GetNameSafe(Animation.Get()), StartTime, EndTime, PreviousPosition, CurrentPosition, DeltaSeconds, CharacterOwner->GetWorld()->GetTimeSeconds(), EndTimeOffset,
+			CurrentDist2D, (GetTargetLocation().Z - CurrentLocation.Z), FutureDist2D, (GetTargetLocation().Z - FutureLocation.Z), Dot,
 			*OriginalRootMotion.GetTranslation().ToString(), OriginalRootMotion.GetTranslation().Size(), *WarpedRootMotion.GetTranslation().ToString(), WarpedRootMotion.GetTranslation().Size(), Speed,
 			*CurrentLocation.ToString(), *FutureLocation.ToString(), *CurrentRotation.ToCompactString(), *FutureRotation.ToCompactString());
 	}
diff --git a/Engine/Plugins/Experimental/Animation/MotionWarping/Source/MotionWarping/Private/RootMotionModifier_AdjustmentBlendWarp.cpp b/Engine/Plugins/Experimental/Animation/MotionWarping/Source/MotionWarping/Private/RootMotionModifier_AdjustmentBlendWarp.cpp
--- a/Engine/Plugins/Experimental/Animation/MotionWarping/Source/MotionWarping/Private/RootMotionModifier_AdjustmentBlendWarp.cpp
+++ b/Engine/Plugins/Experimental/Animation/MotionWarping/Source/MotionWarping/Private/RootMotionModifier_AdjustmentBlendWarp.cpp
@@ -84,7 +84,7 @@ FTransform URootMotionModifier_AdjustmentBlendWarp::ProcessRootMotion(const FTra
 #if !(UE_BUILD_SHIPPING || UE_BUILD_TEST)
 	if (FMotionWarpingCVars::CVarMotionWarpingDebug.GetValueOnGameThread() > 0)
 	{
-		PrintLog(TEXT("FRootMotionModifier_AdjustmentBlend"), InRootMotion, FinalRootMotion);
+		PrintLog(TEXT("AdjustmentBlendWarp"), InRootMotion, FinalRootMotion);
 
 		if (FMotionWarpingCVars::CVarMotionWarpingDebug.GetValueOnGameThread() >= 2)
 		{
diff --git a/Engine/Plugins/Experimental/Animation/MotionWarping/Source/MotionWarping/Private/RootMotionModifier_SkewWarp.cpp b/Engine/Plugins/Experimental/Animation/MotionWarping/Source/MotionWarping/Private/RootMotionModifier_SkewWarp.cpp
--- a/Engine/Plugins/Experimental/Animation/MotionWarping/Source/MotionWarping/Private/RootMotionModifier_SkewWarp.cpp
+++ b/Engine/Plugins/Experimental/Animation/MotionWarping/Source/MotionWarping/Private/RootMotionModifier_SkewWarp.cpp
@@ -5,6 +5,7 @@
 #include "GameFramework/Character.h"
 #include "GameFramework/CharacterMovementComponent.h"
 #include "Components/SkeletalMeshComponent.h"
+#include "Components/CapsuleComponent.h"
 #include "MotionWarpingComponent.h"
 
 URootMotionModifier_SkewWarp::URootMotionModifier_SkewWarp(const FObjectInitializer& ObjectInitializer)
@@ -29,12 +30,10 @@ FTransform URootMotionModifier_SkewWarp::ProcessRootMotion(const FTransform& InR
 	{
 		const FTransform CurrentTransform = FTransform(
 			CharacterOwner->GetActorQuat(),
-			CharacterOwner->GetActorLocation() - FVector::UpVector * CharacterOwner->GetSimpleCollisionHalfHeight()
+			CharacterOwner->GetActorLocation() - FVector(0.f, 0.f, CharacterOwner->GetCapsuleComponent()->GetScaledCapsuleHalfHeight())
 		);
 
-		const FTransform MeshRelativeTransform = FTransform(CharacterOwner->GetBaseRotationOffset(), CharacterOwner->GetBaseTranslationOffset());
-		const FTransform MeshTransform = MeshRelativeTransform * CharacterOwner->GetActorTransform();
-		const FTransform RootMotionTotalWorldSpace = RootMotionTotal * MeshTransform;
+		const FTransform RootMotionTotalWorldSpace = CurrentTransform * CharacterOwner->GetMesh()->ConvertLocalRootMotionToWorld(RootMotionTotal);
 		const FTransform RootMotionDeltaWorldSpace = CharacterOwner->GetMesh()->ConvertLocalRootMotionToWorld(RootMotionDelta);
 
 		const FVector CurrentLocation = CurrentTransform.GetLocation();
@@ -155,7 +154,7 @@ FTransform URootMotionModifier_SkewWarp::ProcessRootMotion(const FTransform& InR
 	const int32 DebugLevel = FMotionWarpingCVars::CVarMotionWarpingDebug.GetValueOnGameThread();
 	if (DebugLevel > 0)
 	{
-		PrintLog(TEXT("FRootMotionModifier_Skew"), InRootMotion, FinalRootMotion);
+		PrintLog(TEXT("SkewWarp"), InRootMotion, FinalRootMotion);
 
 		if (DebugLevel >= 2)
 		{
