Motion Warping API clean up (deprecated simple warp and replaced SyncPoint with WarpTarget with an extra option to create warp targets from a scene comp, an optional bone, and an option to automatically follow that target)

diff --git a/Engine/Plugins/Experimental/Animation/ContextualAnimation/Source/ContextualAnimation/Private/ContextualAnimSceneInstance.cpp b/Engine/Plugins/Experimental/Animation/ContextualAnimation/Source/ContextualAnimation/Private/ContextualAnimSceneInstance.cpp
--- a/Engine/Plugins/Experimental/Animation/ContextualAnimation/Source/ContextualAnimation/Private/ContextualAnimSceneInstance.cpp
+++ b/Engine/Plugins/Experimental/Animation/ContextualAnimation/Source/ContextualAnimation/Private/ContextualAnimSceneInstance.cpp
@@ -246,7 +246,7 @@ void UContextualAnimSceneInstance::Join(FContextualAnimSceneActorData& SceneActo
 				const FName SyncPointName = Pair.Key;
 				const float SyncTime = SceneActorData.GetAnimData()->GetSyncTimeForWarpSection(SyncPointName);
 				const FTransform AlignmentTransform = (SceneActorData.GetAnimData()->AlignmentData.ExtractTransformAtTime(SyncPointName, SyncTime) * Pair.Value);
-				MotionWarpingComp->AddOrUpdateSyncPoint(SyncPointName, AlignmentTransform);
+				MotionWarpingComp->AddOrUpdateWarpTargetFromTransform(SyncPointName, AlignmentTransform);
 			}
 		}
 
diff --git a/Engine/Plugins/Experimental/Animation/MotionWarping/Source/MotionWarping/Private/AnimNotifyState_MotionWarping.cpp b/Engine/Plugins/Experimental/Animation/MotionWarping/Source/MotionWarping/Private/AnimNotifyState_MotionWarping.cpp
--- a/Engine/Plugins/Experimental/Animation/MotionWarping/Source/MotionWarping/Private/AnimNotifyState_MotionWarping.cpp
+++ b/Engine/Plugins/Experimental/Animation/MotionWarping/Source/MotionWarping/Private/AnimNotifyState_MotionWarping.cpp
@@ -3,6 +3,13 @@
 #include "AnimNotifyState_MotionWarping.h"
 #include "GameFramework/Actor.h"
 #include "MotionWarpingComponent.h"
+#include "RootMotionModifier.h"
+
+#if WITH_EDITOR
+#include "Logging/MessageLog.h"
+#include "Misc/UObjectToken.h"
+#include "UObject/ObjectSaveContext.h"
+#endif
 
 UAnimNotifyState_MotionWarping::UAnimNotifyState_MotionWarping(const FObjectInitializer& ObjectInitializer)
 	: Super(ObjectInitializer)
@@ -58,4 +65,51 @@ void UAnimNotifyState_MotionWarping::OnRootMotionModifierDeactivate(UMotionWarpi
 {
 	// Notify blueprint
 	OnWarpEnd(MotionWarpingComp, Modifier);
-}
\ No newline at end of file
+}
+
+#if WITH_EDITOR
+void UAnimNotifyState_MotionWarping::ValidateAssociatedAssets()
+{
+	static const FName NAME_AssetCheck("AssetCheck");
+
+	if (RootMotionModifier == nullptr)
+	{
+		UObject* ContainingAsset = GetContainingAsset();
+
+		FMessageLog AssetCheckLog(NAME_AssetCheck);
+
+		const FText MessageLooping = FText::Format(
+			NSLOCTEXT("AnimNotify", "MotionWarping_InvalidRootMotionModifier", "Motion Warping window in {0} doesn't have a valid RootMotionModifier"),
+			FText::AsCultureInvariant(GetNameSafe(ContainingAsset)));
+		AssetCheckLog.Warning()
+			->AddToken(FUObjectToken::Create(ContainingAsset))
+			->AddToken(FTextToken::Create(MessageLooping));
+
+		if (GIsEditor)
+		{
+			AssetCheckLog.Notify(MessageLooping, EMessageSeverity::Warning, /*bForce=*/ true);
+		}
+	}
+	else if(const URootMotionModifier_Warp* RootMotionModifierWarp = Cast<URootMotionModifier_Warp>(RootMotionModifier))
+	{
+		if(RootMotionModifierWarp->WarpTargetName.IsNone())
+		{
+			UObject* ContainingAsset = GetContainingAsset();
+
+			FMessageLog AssetCheckLog(NAME_AssetCheck);
+
+			const FText MessageLooping = FText::Format(
+				NSLOCTEXT("AnimNotify", "MotionWarping_InvalidWarpTargetName", "Motion Warping window in {0} doesn't specify a valid Warp Target Name"),
+				FText::AsCultureInvariant(GetNameSafe(ContainingAsset)));
+			AssetCheckLog.Warning()
+				->AddToken(FUObjectToken::Create(ContainingAsset))
+				->AddToken(FTextToken::Create(MessageLooping));
+
+			if (GIsEditor)
+			{
+				AssetCheckLog.Notify(MessageLooping, EMessageSeverity::Warning, /*bForce=*/ true);
+			}
+		}
+	}
+}
+#endif
\ No newline at end of file
diff --git a/Engine/Plugins/Experimental/Animation/MotionWarping/Source/MotionWarping/Private/MotionWarpingComponent.cpp b/Engine/Plugins/Experimental/Animation/MotionWarping/Source/MotionWarping/Private/MotionWarpingComponent.cpp
--- a/Engine/Plugins/Experimental/Animation/MotionWarping/Source/MotionWarping/Private/MotionWarpingComponent.cpp
+++ b/Engine/Plugins/Experimental/Animation/MotionWarping/Source/MotionWarping/Private/MotionWarpingComponent.cpp
@@ -16,24 +16,13 @@ DEFINE_LOG_CATEGORY(LogMotionWarping);
 
 #if !(UE_BUILD_SHIPPING || UE_BUILD_TEST)
 TAutoConsoleVariable<int32> FMotionWarpingCVars::CVarMotionWarpingDisable(TEXT("a.MotionWarping.Disable"), 0, TEXT("Disable Motion Warping"), ECVF_Cheat);
-TAutoConsoleVariable<int32> FMotionWarpingCVars::CVarMotionWarpingDebug(TEXT("a.MotionWarping.Debug"), 0, TEXT("0: Disable, 1: Only Log 2: Log and DrawDebug"), ECVF_Cheat);
-TAutoConsoleVariable<float> FMotionWarpingCVars::CVarMotionWarpingDrawDebugDuration(TEXT("a.MotionWarping.DrawDebugLifeTime"), 0.f, TEXT("Time in seconds each draw debug persists.\nRequires 'a.MotionWarping.Debug 2'"), ECVF_Cheat);
+TAutoConsoleVariable<int32> FMotionWarpingCVars::CVarMotionWarpingDebug(TEXT("a.MotionWarping.Debug"), 0, TEXT("0: Disable, 1: Only Log, 2: Only DrawDebug, 3: Log and DrawDebug"), ECVF_Cheat);
+TAutoConsoleVariable<float> FMotionWarpingCVars::CVarMotionWarpingDrawDebugDuration(TEXT("a.MotionWarping.DrawDebugLifeTime"), 1.f, TEXT("Time in seconds each draw debug persists.\nRequires 'a.MotionWarping.Debug 2'"), ECVF_Cheat);
 #endif
 
 // UMotionWarpingUtilities
 ///////////////////////////////////////////////////////////////////////
 
-void UMotionWarpingUtilities::BreakMotionWarpingSyncPoint(const FMotionWarpingSyncPoint& SyncPoint, FVector& Location, FRotator& Rotation)
-{
-	Location = SyncPoint.GetLocation(); 
-	Rotation = SyncPoint.Rotator();
-}
-
-FMotionWarpingSyncPoint UMotionWarpingUtilities::MakeMotionWarpingSyncPoint(FVector Location, FRotator Rotation)
-{
-	return FMotionWarpingSyncPoint(Location, Rotation);
-}
-
 void UMotionWarpingUtilities::ExtractLocalSpacePose(const UAnimSequenceBase* Animation, const FBoneContainer& BoneContainer, float Time, bool bExtractRootMotion, FCompactPose& OutPose)
 {
 	OutPose.SetBoneContainer(&BoneContainer);
@@ -120,7 +109,7 @@ void UMotionWarpingUtilities::GetMotionWarpingWindowsFromAnimation(const UAnimSe
 	}
 }
 
-void UMotionWarpingUtilities::GetMotionWarpingWindowsForSyncPointFromAnimation(const UAnimSequenceBase* Animation, FName WarpTargetName, TArray<FMotionWarpingWindowData>& OutWindows)
+void UMotionWarpingUtilities::GetMotionWarpingWindowsForWarpTargetFromAnimation(const UAnimSequenceBase* Animation, FName WarpTargetName, TArray<FMotionWarpingWindowData>& OutWindows)
 {
 	if (Animation && WarpTargetName != NAME_None)
 	{
@@ -220,6 +209,12 @@ void UMotionWarpingComponent::Update()
 			const UAnimNotifyState_MotionWarping* MotionWarpingNotify = NotifyEvent.NotifyStateClass ? Cast<UAnimNotifyState_MotionWarping>(NotifyEvent.NotifyStateClass) : nullptr;
 			if (MotionWarpingNotify)
 			{
+				if(MotionWarpingNotify->RootMotionModifier == nullptr)
+				{
+					UE_LOG(LogMotionWarping, Warning, TEXT("MotionWarpingComponent::Update. A motion warping window in %s doesn't have a valid root motion modifier!"), *GetNameSafe(Montage));
+					continue;
+				}
+
 				const float StartTime = FMath::Clamp(NotifyEvent.GetTriggerTime(), 0.f, Montage->GetPlayLength());
 				const float EndTime = FMath::Clamp(NotifyEvent.GetEndTriggerTime(), 0.f, Montage->GetPlayLength());
 
@@ -247,6 +242,12 @@ void UMotionWarpingComponent::Update()
 						const UAnimNotifyState_MotionWarping* MotionWarpingNotify = NotifyEvent.NotifyStateClass ? Cast<UAnimNotifyState_MotionWarping>(NotifyEvent.NotifyStateClass) : nullptr;
 						if (MotionWarpingNotify)
 						{
+							if (MotionWarpingNotify->RootMotionModifier == nullptr)
+							{
+								UE_LOG(LogMotionWarping, Warning, TEXT("MotionWarpingComponent::Update. A motion warping window in %s doesn't have a valid root motion modifier!"), *GetNameSafe(AnimSegment->AnimReference));
+								continue;
+							}
+
 							const float NotifyStartTime = FMath::Clamp(NotifyEvent.GetTriggerTime(), 0.f, AnimSegment->AnimReference->GetPlayLength());
 							const float NotifyEndTime = FMath::Clamp(NotifyEvent.GetEndTriggerTime(), 0.f, AnimSegment->AnimReference->GetPlayLength());
 
@@ -355,7 +356,7 @@ FTransform UMotionWarpingComponent::ProcessRootMotionPostConvertToWorld(const FT
 				OriginalRootMotionAccum = ActorFeetLocation;
 				WarpedRootMotionAccum = ActorFeetLocation;
 			}
-			
+
 			OriginalRootMotionAccum = OriginalRootMotionAccum.GetValue() + InRootMotion.GetLocation();
 			WarpedRootMotionAccum = WarpedRootMotionAccum.GetValue() + FinalRootMotion.GetLocation();
 
@@ -375,18 +376,34 @@ FTransform UMotionWarpingComponent::ProcessRootMotionPostConvertToWorld(const FT
 	return FinalRootMotion;
 }
 
-void UMotionWarpingComponent::AddOrUpdateSyncPoint(FName Name, const FMotionWarpingSyncPoint& SyncPoint)
+void UMotionWarpingComponent::AddOrUpdateWarpTarget(FName WarpTargetName, const FMotionWarpingTarget& WarpTarget)
 {
-	if (Name != NAME_None)
+	if (WarpTargetName != NAME_None)
 	{
-		FMotionWarpingSyncPoint& MotionWarpingSyncPoint = SyncPoints.FindOrAdd(Name);
-		MotionWarpingSyncPoint = SyncPoint;
+		FMotionWarpingTarget& WarpTargetRef = WarpTargetMap.FindOrAdd(WarpTargetName);
+		WarpTargetRef = WarpTarget;
 	}
 }
 
-int32 UMotionWarpingComponent::RemoveSyncPoint(FName Name)
+int32 UMotionWarpingComponent::RemoveWarpTarget(FName WarpTargetName)
+{
+	return WarpTargetMap.Remove(WarpTargetName);
+}
+
+void UMotionWarpingComponent::AddOrUpdateWarpTargetFromTransform(FName WarpTargetName, FTransform TargetTransform)
 {
-	return SyncPoints.Remove(Name);
+	AddOrUpdateWarpTarget(WarpTargetName, TargetTransform);
+}
+
+void UMotionWarpingComponent::AddOrUpdateWarpTargetFromComponent(FName WarpTargetName, const USceneComponent* Component, FName BoneName, bool bFollowComponent)
+{
+	if (Component == nullptr)
+	{
+		UE_LOG(LogMotionWarping, Warning, TEXT("AddOrUpdateWarpTargetFromComponent has failed!. Reason: Invalid Component"));
+		return;
+	}
+
+	AddOrUpdateWarpTarget(WarpTargetName, FMotionWarpingTarget(Component, BoneName, bFollowComponent));
 }
 
 URootMotionModifier* UMotionWarpingComponent::AddModifierFromTemplate(URootMotionModifier* Template, const UAnimSequenceBase* Animation, float StartTime, float EndTime)
diff --git a/Engine/Plugins/Experimental/Animation/MotionWarping/Source/MotionWarping/Private/RootMotionModifier.cpp b/Engine/Plugins/Experimental/Animation/MotionWarping/Source/MotionWarping/Private/RootMotionModifier.cpp
--- a/Engine/Plugins/Experimental/Animation/MotionWarping/Source/MotionWarping/Private/RootMotionModifier.cpp
+++ b/Engine/Plugins/Experimental/Animation/MotionWarping/Source/MotionWarping/Private/RootMotionModifier.cpp
@@ -11,6 +11,67 @@
 #include "Animation/AnimInstance.h"
 #include "AnimNotifyState_MotionWarping.h"
 
+// FMotionWarpingTarget
+///////////////////////////////////////////////////////////////
+
+FMotionWarpingTarget::FMotionWarpingTarget(const USceneComponent* InComp, FName InBoneName, bool bInFollowComponent)
+{
+	if (ensure(InComp))
+	{
+		Component = InComp;
+		BoneName = InBoneName;
+		bFollowComponent = bInFollowComponent;
+
+		if (BoneName != NAME_None)
+		{
+			Transform = FMotionWarpingTarget::GetTargetTransformFromComponent(InComp, InBoneName);
+		}
+		else
+		{
+			Transform = InComp->GetComponentTransform();
+		}
+	}
+}
+
+FTransform FMotionWarpingTarget::GetTargetTransformFromComponent(const USceneComponent* Comp, const FName& BoneName)
+{
+	if (Comp == nullptr)
+	{
+		UE_LOG(LogMotionWarping, Warning, TEXT("FMotionWarpingTarget::GetTargetTransformFromComponent: Invalid Component"));
+		return FTransform::Identity;
+	}
+
+	if (Comp->DoesSocketExist(BoneName) == false)
+	{
+		UE_LOG(LogMotionWarping, Warning, TEXT("FMotionWarpingTarget::GetTargetTransformFromComponent: Invalid Bone or Socket. Comp: %s Owner: %s BoneName: %s"),
+			*GetNameSafe(Comp), *GetNameSafe(Comp->GetOwner()), *BoneName.ToString());
+
+		return Comp->GetComponentTransform();
+	}
+
+	return Comp->GetSocketTransform(BoneName);
+}
+
+FTransform FMotionWarpingTarget::GetTargetTrasform() const
+{
+	if (Component.IsValid() && bFollowComponent)
+	{
+		if (BoneName != NAME_None)
+		{
+			return FMotionWarpingTarget::GetTargetTransformFromComponent(Component.Get(), BoneName);
+		}
+		else
+		{
+			return Component->GetComponentTransform();
+		}
+	}
+
+	return Transform;
+}
+
+// URootMotionModifier
+///////////////////////////////////////////////////////////////
+
 URootMotionModifier::URootMotionModifier(const FObjectInitializer& ObjectInitializer)
 	: Super(ObjectInitializer)
 {
@@ -100,6 +161,14 @@ void URootMotionModifier::OnStateChanged(ERootMotionModifierState LastState)
 	{
 		if (LastState != ERootMotionModifierState::Active && State == ERootMotionModifierState::Active)
 		{
+			const ACharacter* CharacterOwner = OwnerComp->GetCharacterOwner();
+			check(CharacterOwner);
+			
+			ActualStartTime = PreviousPosition;
+
+			const float CapsuleHalfHeight = CharacterOwner->GetCapsuleComponent()->GetScaledCapsuleHalfHeight();
+			StartTransform = FTransform(CharacterOwner->GetActorQuat(), (CharacterOwner->GetActorLocation() - FVector(0.f, 0.f, CapsuleHalfHeight)));
+
 			OnActivateDelegate.ExecuteIfBound(OwnerComp, this);
 		}
 		else if (LastState == ERootMotionModifierState::Active && (State == ERootMotionModifierState::Disabled || State == ERootMotionModifierState::MarkedForRemoval))
@@ -119,16 +188,16 @@ void URootMotionModifier_Warp::Update()
 	// Update playback times and state
 	Super::Update();
 
-	// Cache sync point transform and trigger OnSyncPointChanged if needed
+	// Cache sync point transform and trigger OnTargetTransformChanged if needed
 	const UMotionWarpingComponent* OwnerComp = GetOwnerComponent();
 	if (OwnerComp && GetState() == ERootMotionModifierState::Active)
 	{
-		const FMotionWarpingSyncPoint* SyncPointPtr = OwnerComp->FindSyncPoint(WarpTargetName);
+		const FMotionWarpingTarget* WarpTargetPtr = OwnerComp->FindWarpTarget(WarpTargetName);
 
-		// Disable if there is no sync point for us
-		if (SyncPointPtr == nullptr)
+		// Disable if there is no target for us
+		if (WarpTargetPtr == nullptr)
 		{
-			UE_LOG(LogMotionWarping, Verbose, TEXT("MotionWarping: Marking RootMotionModifier as Disabled. Reason: Invalid Sync Point (%s). Char: %s Animation: %s [%f %f] [%f %f]"),
+			UE_LOG(LogMotionWarping, Verbose, TEXT("MotionWarping: Marking RootMotionModifier as Disabled. Reason: Invalid Warp Target (%s). Char: %s Animation: %s [%f %f] [%f %f]"),
 				*WarpTargetName.ToString(), *GetNameSafe(OwnerComp->GetOwner()), *GetNameSafe(Animation.Get()), StartTime, EndTime, PreviousPosition, CurrentPosition);
 
 			SetState(ERootMotionModifierState::Disabled);
@@ -136,7 +205,7 @@ void URootMotionModifier_Warp::Update()
 		}
 
 		// Get the warp point sent by the game
-		FTransform WarpPointTransformGame = FTransform(SyncPointPtr->GetRotation(), SyncPointPtr->GetLocation());
+		FTransform WarpPointTransformGame = WarpTargetPtr->GetTargetTrasform();
 
 		// Initialize our target transform (where the root should end at the end of the window) with the warp point sent by the game
 		FTransform TargetTransform = WarpPointTransformGame;
@@ -202,85 +271,6 @@ void URootMotionModifier_Warp::Update()
 	}
 }
 
-FTransform URootMotionModifier_Warp::ProcessRootMotion(const FTransform& InRootMotion, float DeltaSeconds)
-{
-	const ACharacter* CharacterOwner = GetCharacterOwner();
-	if (CharacterOwner == nullptr)
-	{
-		return InRootMotion;
-	}
-
-	const FTransform& CharacterTransform = CharacterOwner->GetActorTransform();
-
-	FTransform FinalRootMotion = InRootMotion;
-
-	const FTransform RootMotionTotal = UMotionWarpingUtilities::ExtractRootMotionFromAnimation(Animation.Get(), PreviousPosition, EndTime);
-
-	if (bWarpTranslation)
-	{
-		FVector DeltaTranslation = InRootMotion.GetTranslation();
-
-		const FTransform RootMotionDelta = UMotionWarpingUtilities::ExtractRootMotionFromAnimation(Animation.Get(), PreviousPosition, FMath::Min(CurrentPosition, EndTime));
-
-		const float HorizontalDelta = RootMotionDelta.GetTranslation().Size2D();
-		const float HorizontalTarget = FVector::Dist2D(CharacterTransform.GetLocation(), GetTargetLocation());
-		const float HorizontalOriginal = RootMotionTotal.GetTranslation().Size2D();
-		const float HorizontalTranslationWarped = !FMath::IsNearlyZero(HorizontalOriginal) ? ((HorizontalDelta * HorizontalTarget) / HorizontalOriginal) : 0.f;
-
-		if (bInLocalSpace)
-		{
-			const FTransform MeshRelativeTransform = FTransform(CharacterOwner->GetBaseRotationOffset(), CharacterOwner->GetBaseTranslationOffset());
-			const FTransform MeshTransform = MeshRelativeTransform * CharacterOwner->GetActorTransform();
-			DeltaTranslation = MeshTransform.InverseTransformPositionNoScale(GetTargetLocation()).GetSafeNormal2D() * HorizontalTranslationWarped;
-		}
-		else
-		{
-			DeltaTranslation = (GetTargetLocation() - CharacterTransform.GetLocation()).GetSafeNormal2D() * HorizontalTranslationWarped;
-		}
-
-		if (!bIgnoreZAxis)
-		{
-			const float CapsuleHalfHeight = CharacterOwner->GetCapsuleComponent()->GetScaledCapsuleHalfHeight();
-			const FVector CapsuleBottomLocation = (CharacterOwner->GetActorLocation() - FVector(0.f, 0.f, CapsuleHalfHeight));
-			const float VerticalDelta = RootMotionDelta.GetTranslation().Z;
-			const float VerticalTarget = GetTargetLocation().Z - CapsuleBottomLocation.Z;
-			const float VerticalOriginal = RootMotionTotal.GetTranslation().Z;
-			const float VerticalTranslationWarped = !FMath::IsNearlyZero(VerticalOriginal) ? ((VerticalDelta * VerticalTarget) / VerticalOriginal) : 0.f;
-
-			DeltaTranslation.Z = VerticalTranslationWarped;
-		}
-		else
-		{
-			DeltaTranslation.Z = InRootMotion.GetTranslation().Z;
-		}
-
-		FinalRootMotion.SetTranslation(DeltaTranslation);
-	}
-
-	if (bWarpRotation)
-	{
-		const FQuat WarpedRotation = WarpRotation(InRootMotion, RootMotionTotal, DeltaSeconds);
-		FinalRootMotion.SetRotation(WarpedRotation);
-	}
-
-	// Debug
-#if !(UE_BUILD_SHIPPING || UE_BUILD_TEST)
-	const int32 DebugLevel = FMotionWarpingCVars::CVarMotionWarpingDebug.GetValueOnGameThread();
-	if (DebugLevel > 0)
-	{
-		PrintLog(TEXT("SimpleWarp"), InRootMotion, FinalRootMotion);
-
-		if (DebugLevel >= 2)
-		{
-			const float DrawDebugDuration = FMotionWarpingCVars::CVarMotionWarpingDrawDebugDuration.GetValueOnGameThread();
-			DrawDebugCoordinateSystem(CharacterOwner->GetWorld(), GetTargetLocation(), GetTargetRotator(), 50.f, false, DrawDebugDuration, 0, 1.f);
-		}
-	}
-#endif
-
-	return FinalRootMotion;
-}
-
 FQuat URootMotionModifier_Warp::GetTargetRotation() const
 {
 	if (RotationType == EMotionWarpRotationType::Default)
@@ -348,30 +338,89 @@ void URootMotionModifier_Warp::PrintLog(const FString& Name, const FTransform& O
 }
 #endif
 
-URootMotionModifier_Warp* URootMotionModifier_Warp::AddRootMotionModifierSimpleWarp(UMotionWarpingComponent* InMotionWarpingComp, const UAnimSequenceBase* InAnimation, float InStartTime, float InEndTime,
-	FName InWarpTargetName, EWarpPointAnimProvider InWarpPointAnimProvider, FTransform InWarpPointAnimTransform, FName InWarpPointAnimBoneName,
-	bool bInWarpTranslation, bool bInIgnoreZAxis, bool bInWarpRotation, EMotionWarpRotationType InRotationType, float InWarpRotationTimeMultiplier)
+// URootMotionModifier_SimpleWarp
+///////////////////////////////////////////////////////////////
+
+UDEPRECATED_RootMotionModifier_SimpleWarp::UDEPRECATED_RootMotionModifier_SimpleWarp(const FObjectInitializer& ObjectInitializer)
+	: Super(ObjectInitializer)
+{
+}
+
+FTransform UDEPRECATED_RootMotionModifier_SimpleWarp::ProcessRootMotion(const FTransform& InRootMotion, float DeltaSeconds)
 {
-	if (ensureAlways(InMotionWarpingComp))
+	const ACharacter* CharacterOwner = GetCharacterOwner();
+	if (CharacterOwner == nullptr)
+	{
+		return InRootMotion;
+	}
+
+	const FTransform& CharacterTransform = CharacterOwner->GetActorTransform();
+
+	FTransform FinalRootMotion = InRootMotion;
+
+	const FTransform RootMotionTotal = UMotionWarpingUtilities::ExtractRootMotionFromAnimation(Animation.Get(), PreviousPosition, EndTime);
+
+	if (bWarpTranslation)
+	{
+		FVector DeltaTranslation = InRootMotion.GetTranslation();
+
+		const FTransform RootMotionDelta = UMotionWarpingUtilities::ExtractRootMotionFromAnimation(Animation.Get(), PreviousPosition, FMath::Min(CurrentPosition, EndTime));
+
+		const float HorizontalDelta = RootMotionDelta.GetTranslation().Size2D();
+		const float HorizontalTarget = FVector::Dist2D(CharacterTransform.GetLocation(), GetTargetLocation());
+		const float HorizontalOriginal = RootMotionTotal.GetTranslation().Size2D();
+		const float HorizontalTranslationWarped = !FMath::IsNearlyZero(HorizontalOriginal) ? ((HorizontalDelta * HorizontalTarget) / HorizontalOriginal) : 0.f;
+
+		if (bInLocalSpace)
+		{
+			const FTransform MeshRelativeTransform = FTransform(CharacterOwner->GetBaseRotationOffset(), CharacterOwner->GetBaseTranslationOffset());
+			const FTransform MeshTransform = MeshRelativeTransform * CharacterOwner->GetActorTransform();
+			DeltaTranslation = MeshTransform.InverseTransformPositionNoScale(GetTargetLocation()).GetSafeNormal2D() * HorizontalTranslationWarped;
+		}
+		else
+		{
+			DeltaTranslation = (GetTargetLocation() - CharacterTransform.GetLocation()).GetSafeNormal2D() * HorizontalTranslationWarped;
+		}
+
+		if (!bIgnoreZAxis)
+		{
+			const float CapsuleHalfHeight = CharacterOwner->GetCapsuleComponent()->GetScaledCapsuleHalfHeight();
+			const FVector CapsuleBottomLocation = (CharacterOwner->GetActorLocation() - FVector(0.f, 0.f, CapsuleHalfHeight));
+			const float VerticalDelta = RootMotionDelta.GetTranslation().Z;
+			const float VerticalTarget = GetTargetLocation().Z - CapsuleBottomLocation.Z;
+			const float VerticalOriginal = RootMotionTotal.GetTranslation().Z;
+			const float VerticalTranslationWarped = !FMath::IsNearlyZero(VerticalOriginal) ? ((VerticalDelta * VerticalTarget) / VerticalOriginal) : 0.f;
+
+			DeltaTranslation.Z = VerticalTranslationWarped;
+		}
+		else
+		{
+			DeltaTranslation.Z = InRootMotion.GetTranslation().Z;
+		}
+
+		FinalRootMotion.SetTranslation(DeltaTranslation);
+	}
+
+	if (bWarpRotation)
 	{
-		URootMotionModifier_Warp* NewModifier = NewObject<URootMotionModifier_Warp>(InMotionWarpingComp);
-		NewModifier->Animation = InAnimation;
-		NewModifier->StartTime = InStartTime;
-		NewModifier->EndTime = InEndTime;
-		NewModifier->WarpTargetName = InWarpTargetName;
-		NewModifier->WarpPointAnimProvider = InWarpPointAnimProvider;
-		NewModifier->WarpPointAnimTransform = InWarpPointAnimTransform;
-		NewModifier->WarpPointAnimBoneName = InWarpPointAnimBoneName;
-		NewModifier->bWarpTranslation = bInWarpTranslation;
-		NewModifier->bIgnoreZAxis = bInIgnoreZAxis;
-		NewModifier->bWarpRotation = bInWarpRotation;
-		NewModifier->RotationType = InRotationType;
-		NewModifier->WarpRotationTimeMultiplier = InWarpRotationTimeMultiplier;
-		
-		InMotionWarpingComp->AddModifier(NewModifier);
-
-		return NewModifier;
+		const FQuat WarpedRotation = WarpRotation(InRootMotion, RootMotionTotal, DeltaSeconds);
+		FinalRootMotion.SetRotation(WarpedRotation);
 	}
 
-	return nullptr;
+	// Debug
+#if !(UE_BUILD_SHIPPING || UE_BUILD_TEST)
+	const int32 DebugLevel = FMotionWarpingCVars::CVarMotionWarpingDebug.GetValueOnGameThread();
+	if (DebugLevel == 1 || DebugLevel == 3)
+	{
+		PrintLog(TEXT("SimpleWarp"), InRootMotion, FinalRootMotion);
+	}
+
+	if (DebugLevel == 2 || DebugLevel == 3)
+	{
+		const float DrawDebugDuration = FMotionWarpingCVars::CVarMotionWarpingDrawDebugDuration.GetValueOnGameThread();
+		DrawDebugCoordinateSystem(CharacterOwner->GetWorld(), GetTargetLocation(), GetTargetRotator(), 50.f, false, DrawDebugDuration, 0, 1.f);
+	}
+#endif
+
+	return FinalRootMotion;
 }
\ No newline at end of file
diff --git a/Engine/Plugins/Experimental/Animation/MotionWarping/Source/MotionWarping/Private/RootMotionModifier_AdjustmentBlendWarp.cpp b/Engine/Plugins/Experimental/Animation/MotionWarping/Source/MotionWarping/Private/RootMotionModifier_AdjustmentBlendWarp.cpp
--- a/Engine/Plugins/Experimental/Animation/MotionWarping/Source/MotionWarping/Private/RootMotionModifier_AdjustmentBlendWarp.cpp
+++ b/Engine/Plugins/Experimental/Animation/MotionWarping/Source/MotionWarping/Private/RootMotionModifier_AdjustmentBlendWarp.cpp
@@ -82,15 +82,16 @@ FTransform URootMotionModifier_AdjustmentBlendWarp::ProcessRootMotion(const FTra
 
 	// Debug
 #if !(UE_BUILD_SHIPPING || UE_BUILD_TEST)
-	if (FMotionWarpingCVars::CVarMotionWarpingDebug.GetValueOnGameThread() > 0)
+	const int32 DebugLevel = FMotionWarpingCVars::CVarMotionWarpingDebug.GetValueOnGameThread();
+	if (DebugLevel == 1 || DebugLevel == 3)
 	{
 		PrintLog(TEXT("AdjustmentBlendWarp"), InRootMotion, FinalRootMotion);
+	}
 
-		if (FMotionWarpingCVars::CVarMotionWarpingDebug.GetValueOnGameThread() >= 2)
-		{
-			const float DrawDebugDuration = FMotionWarpingCVars::CVarMotionWarpingDrawDebugDuration.GetValueOnGameThread();
-			DrawDebugWarpedTracks(DrawDebugDuration);
-		}
+	if (DebugLevel == 2 || DebugLevel == 3)
+	{
+		const float DrawDebugDuration = FMotionWarpingCVars::CVarMotionWarpingDrawDebugDuration.GetValueOnGameThread();
+		DrawDebugWarpedTracks(DrawDebugDuration);
 	}
 #endif
 
@@ -366,15 +367,15 @@ void URootMotionModifier_AdjustmentBlendWarp::DrawDebugWarpedTracks(float DrawDu
 			{
 				for (int32 FrameIdx = 0; FrameIdx < Result.AnimationTracks[TrackIndex].PosKeys.Num() - 1; FrameIdx++)
 				{
-					FTransform StartTransform;
-					ExtractBoneTransformAtFrame(StartTransform, TrackIndex, FrameIdx);
-					StartTransform = StartTransform * RootStartTransform.Inverse() * CachedMeshTransform;
+					FTransform T1;
+					ExtractBoneTransformAtFrame(T1, TrackIndex, FrameIdx);
+					T1 = T1 * RootStartTransform.Inverse() * CachedMeshTransform;
 
-					FTransform EndTransform;
-					ExtractBoneTransformAtFrame(EndTransform, TrackIndex, FrameIdx + 1);
-					EndTransform = EndTransform * RootStartTransform.Inverse() * CachedMeshTransform;
+					FTransform T2;
+					ExtractBoneTransformAtFrame(T2, TrackIndex, FrameIdx + 1);
+					T2 = T2 * RootStartTransform.Inverse() * CachedMeshTransform;
 
-					DrawDebugLine(World, StartTransform.GetLocation(), EndTransform.GetLocation(), FColor::Yellow, false, DrawDuration, 0, 0.5f);
+					DrawDebugLine(World, T1.GetLocation(), T2.GetLocation(), FColor::Yellow, false, DrawDuration, 0, 0.5f);
 				}
 			}
 		}
diff --git a/Engine/Plugins/Experimental/Animation/MotionWarping/Source/MotionWarping/Private/RootMotionModifier_SkewWarp.cpp b/Engine/Plugins/Experimental/Animation/MotionWarping/Source/MotionWarping/Private/RootMotionModifier_SkewWarp.cpp
--- a/Engine/Plugins/Experimental/Animation/MotionWarping/Source/MotionWarping/Private/RootMotionModifier_SkewWarp.cpp
+++ b/Engine/Plugins/Experimental/Animation/MotionWarping/Source/MotionWarping/Private/RootMotionModifier_SkewWarp.cpp
@@ -152,15 +152,15 @@ FTransform URootMotionModifier_SkewWarp::ProcessRootMotion(const FTransform& InR
 	// Debug
 #if !(UE_BUILD_SHIPPING || UE_BUILD_TEST)
 	const int32 DebugLevel = FMotionWarpingCVars::CVarMotionWarpingDebug.GetValueOnGameThread();
-	if (DebugLevel > 0)
+	if (DebugLevel == 1 || DebugLevel == 3)
 	{
 		PrintLog(TEXT("SkewWarp"), InRootMotion, FinalRootMotion);
+	}
 
-		if (DebugLevel >= 2)
-		{
-			const float DrawDebugDuration = FMotionWarpingCVars::CVarMotionWarpingDrawDebugDuration.GetValueOnGameThread();
-			DrawDebugCoordinateSystem(CharacterOwner->GetWorld(), GetTargetLocation(), GetTargetRotator(), 50.f, false, DrawDebugDuration, 0, 1.f);
-		}
+	if (DebugLevel == 2 || DebugLevel == 3)
+	{
+		const float DrawDebugDuration = FMotionWarpingCVars::CVarMotionWarpingDrawDebugDuration.GetValueOnGameThread();
+		DrawDebugCoordinateSystem(CharacterOwner->GetWorld(), GetTargetLocation(), GetTargetRotator(), 50.f, false, DrawDebugDuration, 0, 1.f);
 	}
 #endif
 
diff --git a/Engine/Plugins/Experimental/Animation/MotionWarping/Source/MotionWarping/Public/AnimNotifyState_MotionWarping.h b/Engine/Plugins/Experimental/Animation/MotionWarping/Source/MotionWarping/Public/AnimNotifyState_MotionWarping.h
--- a/Engine/Plugins/Experimental/Animation/MotionWarping/Source/MotionWarping/Public/AnimNotifyState_MotionWarping.h
+++ b/Engine/Plugins/Experimental/Animation/MotionWarping/Source/MotionWarping/Public/AnimNotifyState_MotionWarping.h
@@ -4,11 +4,11 @@
 
 #include "CoreMinimal.h"
 #include "Animation/AnimNotifies/AnimNotifyState.h"
-#include "RootMotionModifier.h"
 #include "AnimNotifyState_MotionWarping.generated.h"
 
 class UMotionWarpingComponent;
 class UAnimSequenceBase;
+class URootMotionModifier;
 
 /** AnimNotifyState used to define a motion warping window in the animation */
 UCLASS(meta = (DisplayName = "Motion Warping"))
@@ -49,4 +49,8 @@ public:
 
 	UFUNCTION(BlueprintImplementableEvent, Category = "Motion Warping")
 	void OnWarpEnd(UMotionWarpingComponent* MotionWarpingComp, URootMotionModifier* Modifier) const;
+
+#if WITH_EDITOR
+	virtual void ValidateAssociatedAssets() override;
+#endif
 };
\ No newline at end of file
diff --git a/Engine/Plugins/Experimental/Animation/MotionWarping/Source/MotionWarping/Public/MotionWarpingComponent.h b/Engine/Plugins/Experimental/Animation/MotionWarping/Source/MotionWarping/Public/MotionWarpingComponent.h
--- a/Engine/Plugins/Experimental/Animation/MotionWarping/Source/MotionWarping/Public/MotionWarpingComponent.h
+++ b/Engine/Plugins/Experimental/Animation/MotionWarping/Source/MotionWarping/Public/MotionWarpingComponent.h
@@ -47,14 +47,6 @@ class MOTIONWARPING_API UMotionWarpingUtilities : public UBlueprintFunctionLibra
 
 public:
 
-	/** Extracts data from a MotionWarpingSyncPoint */
-	UFUNCTION(BlueprintPure, Category = "Motion Warping", meta = (NativeBreakFunc))
-	static void BreakMotionWarpingSyncPoint(const FMotionWarpingSyncPoint& SyncPoint, FVector& Location, FRotator& Rotation);
-
-	/** Create a MotionWarpingSyncPoint struct */
-	UFUNCTION(BlueprintPure, Category = "Motion Warping", meta = (NativeMakeFunc))
-	static FMotionWarpingSyncPoint MakeMotionWarpingSyncPoint(FVector Location, FRotator Rotation);
-
 	/** Extract bone pose in local space for all bones in BoneContainer. If Animation is a Montage the pose is extracted from the first track */
 	static void ExtractLocalSpacePose(const UAnimSequenceBase* Animation, const FBoneContainer& BoneContainer, float Time, bool bExtractRootMotion, FCompactPose& OutPose);
 
@@ -72,9 +64,9 @@ public:
 	UFUNCTION(BlueprintCallable, Category = "Motion Warping")
 	static void GetMotionWarpingWindowsFromAnimation(const UAnimSequenceBase* Animation, TArray<FMotionWarpingWindowData>& OutWindows);
 
-	/** @return All the MotionWarping windows within the supplied animation for a given Sync Point */
+	/** @return All the MotionWarping windows within the supplied animation for a given Warp Target */
 	UFUNCTION(BlueprintCallable, Category = "Motion Warping")
-	static void GetMotionWarpingWindowsForSyncPointFromAnimation(const UAnimSequenceBase* Animation, FName WarpTargetName, TArray<FMotionWarpingWindowData>& OutWindows);
+	static void GetMotionWarpingWindowsForWarpTargetFromAnimation(const UAnimSequenceBase* Animation, FName WarpTargetName, TArray<FMotionWarpingWindowData>& OutWindows);
 };
 
 DECLARE_DYNAMIC_MULTICAST_DELEGATE_OneParam(FMotionWarpingPreUpdate, class UMotionWarpingComponent*, MotionWarpingComp);
@@ -104,17 +96,6 @@ public:
 	/** Returns the list of root motion modifiers */
 	FORCEINLINE const TArray<URootMotionModifier*>& GetModifiers() const { return Modifiers; }
 
-	/** Find the SyncPoint associated with a specified name */
-	FORCEINLINE const FMotionWarpingSyncPoint* FindSyncPoint(const FName& SyncPointName) const { return SyncPoints.Find(SyncPointName); }
-
-	/** Adds or update the sync point associated with a specified name */
-	UFUNCTION(BlueprintCallable, Category = "Motion Warping")
-	void AddOrUpdateSyncPoint(FName Name, const FMotionWarpingSyncPoint& SyncPoint);
-
-	/** Removes sync point associated with the specified key  */
-	UFUNCTION(BlueprintCallable, Category = "Motion Warping")
-	int32 RemoveSyncPoint(FName Name);
-
 	/** Check if we contain a RootMotionModifier for the supplied animation and time range */
 	bool ContainsModifier(const UAnimSequenceBase* Animation, float StartTime, float EndTime) const;
 
@@ -127,6 +108,58 @@ public:
 
 	URootMotionModifier* AddModifierFromTemplate(URootMotionModifier* Template, const UAnimSequenceBase* Animation, float StartTime, float EndTime);
 
+	/** Find the target associated with a specified name */
+	FORCEINLINE const FMotionWarpingTarget* FindWarpTarget(const FName& WarpTargetName) const { return WarpTargetMap.Find(WarpTargetName); }
+
+	/** Adds or update a target associated with a specified name */
+	UFUNCTION(BlueprintCallable, Category = "Motion Warping")
+	void AddOrUpdateWarpTarget(FName WarpTargetName, const FMotionWarpingTarget& WarpTarget);
+
+	/** 
+	 * Create and adds or update a target associated with a specified name 
+	 * @param WarpTargetName Warp target identifier
+	 * @param TargetTransform Transform used to set the location and rotation for the warp target
+	 */
+	UFUNCTION(BlueprintCallable, Category = "Motion Warping")
+	void AddOrUpdateWarpTargetFromTransform(FName WarpTargetName, FTransform TargetTransform);
+
+	/**
+	 * Create and adds or update a target associated with a specified name
+	 * @param WarpTargetName Warp target identifier
+	 * @param Component Scene Component used to get the target transform
+	 * @param BoneName Optional bone or socket in the component used to get the target transform. 
+	 * @param bFollowComponent Whether the target transform should update while the warping is active. Useful for tracking moving targets.
+	 *		  Note that this will be one frame off if our owner ticks before the target actor. Add tick pre-requisites to avoid this.
+	 */
+	UFUNCTION(BlueprintCallable, Category = "Motion Warping")
+	void AddOrUpdateWarpTargetFromComponent(FName WarpTargetName, const USceneComponent* Component, FName BoneName, bool bFollowComponent);
+
+	/**
+	 * Create and adds or update a target associated with a specified name
+	 * @param WarpTargetName Warp target identifier
+	 * @param TargetLocation Location for the warp target
+	 */
+	UFUNCTION(BlueprintCallable, Category = "Motion Warping")
+	void AddOrUpdateWarpTargetFromLocation(FName WarpTargetName, FVector TargetLocation)
+	{
+		AddOrUpdateWarpTargetFromTransform(WarpTargetName, FTransform(TargetLocation));
+	}
+
+	/**
+	 * Create and adds or update a target associated with a specified name
+	 * @param WarpTargetName Warp target identifier
+	 * @param TargetLocation Location for the warp target
+	 */
+	UFUNCTION(BlueprintCallable, Category = "Motion Warping")
+	void AddOrUpdateWarpTargetFromLocationAndRotation(FName WarpTargetName, FVector TargetLocation, FRotator TargetRotation)
+	{
+		AddOrUpdateWarpTargetFromTransform(WarpTargetName, FTransform(TargetRotation, TargetLocation));
+	}
+
+	/** Removes the warp target associated with the specified key  */
+	UFUNCTION(BlueprintCallable, Category = "Motion Warping")
+	int32 RemoveWarpTarget(FName WarpTargetName);
+
 protected:
 
 	/** Character this component belongs to */
@@ -138,7 +171,7 @@ protected:
 	TArray<URootMotionModifier*> Modifiers;
 
 	UPROPERTY(Transient)
-	TMap<FName, FMotionWarpingSyncPoint> SyncPoints;
+	TMap<FName, FMotionWarpingTarget> WarpTargetMap;
 
 #if !(UE_BUILD_SHIPPING || UE_BUILD_TEST)
 	TOptional<FVector> OriginalRootMotionAccum;
diff --git a/Engine/Plugins/Experimental/Animation/MotionWarping/Source/MotionWarping/Public/RootMotionModifier.h b/Engine/Plugins/Experimental/Animation/MotionWarping/Source/MotionWarping/Public/RootMotionModifier.h
--- a/Engine/Plugins/Experimental/Animation/MotionWarping/Source/MotionWarping/Public/RootMotionModifier.h
+++ b/Engine/Plugins/Experimental/Animation/MotionWarping/Source/MotionWarping/Public/RootMotionModifier.h
@@ -9,6 +9,7 @@
 class UMotionWarpingComponent;
 class UAnimNotifyState_MotionWarping;
 class URootMotionModifier;
+class USceneComponent;
 
 /** The possible states of a Root Motion Modifier */
 UENUM(BlueprintType)
@@ -29,71 +30,9 @@ enum class ERootMotionModifierState : uint8
 
 DECLARE_DYNAMIC_DELEGATE_TwoParams(FOnRootMotionModifierDelegate, UMotionWarpingComponent*, MotionWarpingComp, URootMotionModifier*, RootMotionModifier);
 
-/** Represents a point of alignment in the world */
-USTRUCT(BlueprintType, meta = (HasNativeMake = "MotionWarping.MotionWarpingUtilities.MakeMotionWarpingSyncPoint", HasNativeBreak = "MotionWarping.MotionWarpingUtilities.BreakMotionWarpingSyncPoint"))
-struct MOTIONWARPING_API FMotionWarpingSyncPoint
-{
-	GENERATED_BODY()
-
-	FMotionWarpingSyncPoint()
-		: Location(FVector::ZeroVector), Rotation(FQuat::Identity) {}
-	FMotionWarpingSyncPoint(const FVector& InLocation, const FQuat& InRotation)
-		: Location(InLocation), Rotation(InRotation) {}
-	FMotionWarpingSyncPoint(const FVector& InLocation, const FRotator& InRotation)
-		: Location(InLocation), Rotation(InRotation.Quaternion()) {}
-	FMotionWarpingSyncPoint(const FTransform& InTransform)
-		: Location(InTransform.GetLocation()), Rotation(InTransform.GetRotation()) {}
-
-	FORCEINLINE const FVector& GetLocation() const { return Location; }
-	FORCEINLINE const FQuat& GetRotation() const { return Rotation; }
-	FORCEINLINE FRotator Rotator() const { return Rotation.Rotator(); }
-
-	FORCEINLINE bool operator==(const FMotionWarpingSyncPoint& Other) const
-	{
-		return Other.Location.Equals(Location) && Other.Rotation.Equals(Rotation);
-	}
-
-	FORCEINLINE bool operator!=(const FMotionWarpingSyncPoint& Other) const
-	{
-		return !Other.Location.Equals(Location) || !Other.Rotation.Equals(Rotation);
-	}
-
-protected:
-
-	UPROPERTY()
-	FVector Location;
-
-	UPROPERTY()
-	FQuat Rotation;
-};
-
-// FRootMotionModifier_Warp
+// URootMotionModifier
 ///////////////////////////////////////////////////////////////
 
-UENUM(BlueprintType)
-enum class EMotionWarpRotationType : uint8
-{
-	/** Character rotates to match the rotation of the sync point */
-	Default,
-
-	/** Character rotates to face the sync point */
-	Facing,
-};
-
-/** Method used to extract the warp point from the animation */
-UENUM(BlueprintType)
-enum class EWarpPointAnimProvider : uint8
-{
-	/** No warp point is provided */
-	None,
-
-	/** Warp point defined by a 'hard-coded' transform  user can enter through the warping notify */
-	Static,
-
-	/** Warp point defined by a bone */
-	Bone
-};
-
 UCLASS(Abstract, BlueprintType, EditInlineNew)
 class MOTIONWARPING_API URootMotionModifier : public UObject
 {
@@ -129,6 +68,14 @@ public:
 	UPROPERTY(BlueprintReadOnly, Category = "Defaults")
 	bool bInLocalSpace = false;
 
+	/** Character owner transform at the time this modifier becomes active */
+	UPROPERTY(Transient, BlueprintReadOnly, Category = "Defaults")
+	FTransform StartTransform;
+
+	/** Actual playback time when the modifier becomes active */
+	UPROPERTY(Transient, BlueprintReadOnly, Category = "Defaults")
+	float ActualStartTime = 0.f;
+
 	/** Delegate called when this modifier is activated (starts affecting the root motion) */
 	UPROPERTY()
 	FOnRootMotionModifierDelegate OnActivateDelegate;
@@ -172,7 +119,84 @@ private:
 	ERootMotionModifierState State = ERootMotionModifierState::Waiting;
 };
 
-UCLASS(meta = (DisplayName = "Simple Warp"))
+// URootMotionModifier_Warp
+///////////////////////////////////////////////////////////////
+
+/** Represents a point of alignment in the world */
+USTRUCT(BlueprintType)
+struct MOTIONWARPING_API FMotionWarpingTarget
+{
+	GENERATED_BODY()
+
+	/** When the warp target is created from a component this stores the transform of the component at the time of creation, otherwise its the transform supplied by the user */
+	UPROPERTY(EditAnywhere, BlueprintReadWrite, Category = "Defaults")
+	FTransform Transform;
+
+	/** Optional component used to calculate the final target transform */
+	UPROPERTY(EditAnywhere, BlueprintReadWrite, Category = "Defaults")
+	TWeakObjectPtr<const USceneComponent> Component;
+
+	/** Optional bone name in the component used to calculate the final target transform */
+	UPROPERTY(EditAnywhere, BlueprintReadWrite, Category = "Defaults")
+	FName BoneName;
+
+	/** Whether the target transform calculated from a component and an optional bone should be updated during the warp */
+	UPROPERTY(EditAnywhere, BlueprintReadWrite, Category = "Defaults")
+	bool bFollowComponent;
+
+	FMotionWarpingTarget()
+		: Transform(FTransform::Identity), Component(nullptr), BoneName(NAME_None), bFollowComponent(false) {}
+
+	FMotionWarpingTarget(const FTransform& InTransform)
+		: Transform(InTransform), Component(nullptr), BoneName(NAME_None), bFollowComponent(false) {}
+
+	FMotionWarpingTarget(const USceneComponent* InComp, FName InBoneName, bool bInbFollowComponent);
+
+	FTransform GetTargetTrasform() const;
+
+	FORCEINLINE FVector GetLocation() const { return GetTargetTrasform().GetLocation(); }
+	FORCEINLINE FQuat GetRotation() const { return GetTargetTrasform().GetRotation(); }
+	FORCEINLINE FRotator Rotator() const { return GetTargetTrasform().Rotator(); }
+
+	FORCEINLINE bool operator==(const FMotionWarpingTarget& Other) const
+	{
+		return Other.Transform.Equals(Transform) && Other.Component == Component && Other.BoneName == BoneName && Other.bFollowComponent == bFollowComponent;
+	}
+
+	FORCEINLINE bool operator!=(const FMotionWarpingTarget& Other) const
+	{
+		return !Other.Transform.Equals(Transform) || Other.Component != Component || Other.BoneName != BoneName || Other.bFollowComponent != bFollowComponent;
+	}
+
+	static FTransform GetTargetTransformFromComponent(const USceneComponent* Comp, const FName& BoneName);
+
+};
+
+UENUM(BlueprintType)
+enum class EMotionWarpRotationType : uint8
+{
+	/** Character rotates to match the rotation of the sync point */
+	Default,
+
+	/** Character rotates to face the sync point */
+	Facing,
+};
+
+/** Method used to extract the warp point from the animation */
+UENUM(BlueprintType)
+enum class EWarpPointAnimProvider : uint8
+{
+	/** No warp point is provided */
+	None,
+
+	/** Warp point defined by a 'hard-coded' transform  user can enter through the warping notify */
+	Static,
+
+	/** Warp point defined by a bone */
+	Bone
+};
+
+UCLASS(Abstract)
 class MOTIONWARPING_API URootMotionModifier_Warp : public URootMotionModifier
 {
 	GENERATED_BODY()
@@ -221,7 +245,6 @@ public:
 
 	//~ Begin FRootMotionModifier Interface
 	virtual void Update() override;
-	virtual FTransform ProcessRootMotion(const FTransform& InRootMotion, float DeltaSeconds) override;
 	//~ End FRootMotionModifier Interface
 
 	/** Event called during update if the target transform changes while the warping is active */
@@ -231,29 +254,39 @@ public:
 	void PrintLog(const FString& Name, const FTransform& OriginalRootMotion, const FTransform& WarpedRootMotion) const;
 #endif
 
-	UFUNCTION(BlueprintCallable, Category = "Motion Warping")
-	static URootMotionModifier_Warp* AddRootMotionModifierSimpleWarp(UMotionWarpingComponent* InMotionWarpingComp, const UAnimSequenceBase* InAnimation, float InStartTime, float InEndTime,
-		FName InWarpTargetName, EWarpPointAnimProvider InWarpPointAnimProvider, FTransform InWarpPointAnimTransform, FName InWarpPointAnimBoneName,
-		bool bInWarpTranslation, bool bInIgnoreZAxis, bool bInWarpRotation, EMotionWarpRotationType InRotationType, float InWarpRotationTimeMultiplier = 1.f);
-
-protected:
-
 	FORCEINLINE FVector GetTargetLocation() const { return CachedTargetTransform.GetLocation(); }
 	FORCEINLINE FRotator GetTargetRotator() const { return GetTargetRotation().Rotator(); }
 	FQuat GetTargetRotation() const;
 
 	FQuat WarpRotation(const FTransform& RootMotionDelta, const FTransform& RootMotionTotal, float DeltaSeconds);
 
-	//@TODO: This should be Optional and private
+protected:
+
 	UPROPERTY()
 	FTransform CachedTargetTransform = FTransform::Identity;
 
-private:
-
 	/** Cached of the offset from the warp target. Used to calculate the final target transform when a warp target is defined in the animation */
 	TOptional<FTransform> CachedOffsetFromWarpPoint;
 };
 
+// URootMotionModifier_SimpleWarp. 
+// DEPRECATED in favor of URootMotionModifier_SkewWarp (kept for reference)
+///////////////////////////////////////////////////////////////
+
+UCLASS(Deprecated, meta = (DisplayName = "Simple Warp"))
+class MOTIONWARPING_API UDEPRECATED_RootMotionModifier_SimpleWarp : public URootMotionModifier_Warp
+{
+	GENERATED_BODY()
+
+public:
+
+	UDEPRECATED_RootMotionModifier_SimpleWarp(const FObjectInitializer& ObjectInitializer);
+	virtual FTransform ProcessRootMotion(const FTransform& InRootMotion, float DeltaSeconds) override;
+};
+
+// URootMotionModifier_Scale
+///////////////////////////////////////////////////////////////
+
 UCLASS(meta = (DisplayName = "Scale"))
 class MOTIONWARPING_API URootMotionModifier_Scale : public URootMotionModifier
 {
diff --git a/Engine/Plugins/Experimental/Animation/MotionWarping/Source/MotionWarping/Public/RootMotionModifier_AdjustmentBlendWarp.h b/Engine/Plugins/Experimental/Animation/MotionWarping/Source/MotionWarping/Public/RootMotionModifier_AdjustmentBlendWarp.h
--- a/Engine/Plugins/Experimental/Animation/MotionWarping/Source/MotionWarping/Public/RootMotionModifier_AdjustmentBlendWarp.h
+++ b/Engine/Plugins/Experimental/Animation/MotionWarping/Source/MotionWarping/Public/RootMotionModifier_AdjustmentBlendWarp.h
@@ -83,9 +83,6 @@ protected:
 	UPROPERTY()
 	FAnimSequenceTrackContainer Result;
 
-	UPROPERTY()
-	float ActualStartTime = 0.f;
-
 	void PrecomputeWarpedTracks();
 
 	FTransform ExtractWarpedRootMotion() const;
diff --git a/Engine/Source/Runtime/Engine/Classes/Animation/AnimNotifies/AnimNotifyState.h b/Engine/Source/Runtime/Engine/Classes/Animation/AnimNotifies/AnimNotifyState.h
--- a/Engine/Source/Runtime/Engine/Classes/Animation/AnimNotifies/AnimNotifyState.h
+++ b/Engine/Source/Runtime/Engine/Classes/Animation/AnimNotifies/AnimNotifyState.h
@@ -45,6 +45,7 @@ class ENGINE_API UAnimNotifyState : public UObject
 #if WITH_EDITOR
 	virtual void OnAnimNotifyCreatedInEditor(FAnimNotifyEvent& ContainingAnimNotifyEvent) {};
 	virtual bool CanBePlaced(UAnimSequenceBase* Animation) const { return true; }
+	virtual void ValidateAssociatedAssets() {}
 #endif
 
 	UE_DEPRECATED(5.0, "This function is deprecated. Use the other NotifyBegin instead.")
@@ -83,11 +84,15 @@ class ENGINE_API UAnimNotifyState : public UObject
 	}
 
 	/** UObject Interface */
-	virtual void PostLoad();
+	virtual void PostLoad() override;
+	virtual void PreSave(FObjectPreSaveContext ObjectSaveContext) override;
 	/** End UObject Interface */
 
 	/** This notify is always a branching point when used on Montages. */
 	bool bIsNativeBranchingPoint;
+
+protected:
+	UObject* GetContainingAsset() const;
 };
 
 
diff --git a/Engine/Source/Runtime/Engine/Private/Animation/AnimNotifyState.cpp b/Engine/Source/Runtime/Engine/Private/Animation/AnimNotifyState.cpp
--- a/Engine/Source/Runtime/Engine/Private/Animation/AnimNotifyState.cpp
+++ b/Engine/Source/Runtime/Engine/Private/Animation/AnimNotifyState.cpp
@@ -3,6 +3,7 @@
 #include "Animation/AnimNotifies/AnimNotifyState.h"
 #include "Animation/AnimTypes.h"
 #include "Animation/AnimNotifies/AnimNotify.h"
+#include "UObject/ObjectSaveContext.h"
 
 /////////////////////////////////////////////////////
 // UAnimNotifyState
@@ -108,7 +109,26 @@ void UAnimNotifyState::PostLoad()
 #if WITH_EDITOR
 	// Ensure that all loaded notifies are transactional
 	SetFlags(GetFlags() | RF_Transactional);
+
+	// Make sure the asset isn't bogus (e.g., a looping particle system in a one-shot notify)
+	ValidateAssociatedAssets();
 #endif
 }
 
+void UAnimNotifyState::PreSave(FObjectPreSaveContext ObjectSaveContext)
+{
+#if WITH_EDITOR
+	ValidateAssociatedAssets();
+#endif
+	Super::PreSave(ObjectSaveContext);
+}
 
+UObject* UAnimNotifyState::GetContainingAsset() const
+{
+	UObject* ContainingAsset = GetTypedOuter<UAnimSequenceBase>();
+	if (ContainingAsset == nullptr)
+	{
+		ContainingAsset = GetOutermost();
+	}
+	return ContainingAsset;
+}
