Motion Warping: Do not run motion warping window validation when saving or loading a custom warping window blueprint. This should only run for windows added to an animation.

diff --git a/Engine/Plugins/Experimental/Animation/MotionWarping/Source/MotionWarping/Private/AnimNotifyState_MotionWarping.cpp b/Engine/Plugins/Experimental/Animation/MotionWarping/Source/MotionWarping/Private/AnimNotifyState_MotionWarping.cpp
--- a/Engine/Plugins/Experimental/Animation/MotionWarping/Source/MotionWarping/Private/AnimNotifyState_MotionWarping.cpp
+++ b/Engine/Plugins/Experimental/Animation/MotionWarping/Source/MotionWarping/Private/AnimNotifyState_MotionWarping.cpp
@@ -72,34 +72,14 @@ void UAnimNotifyState_MotionWarping::ValidateAssociatedAssets()
 {
 	static const FName NAME_AssetCheck("AssetCheck");
 
-	if (RootMotionModifier == nullptr)
+	if(UAnimSequenceBase* ContainingAsset = Cast<UAnimSequenceBase>(GetContainingAsset()))
 	{
-		UObject* ContainingAsset = GetContainingAsset();
-
-		FMessageLog AssetCheckLog(NAME_AssetCheck);
-
-		const FText MessageLooping = FText::Format(
-			NSLOCTEXT("AnimNotify", "MotionWarping_InvalidRootMotionModifier", "Motion Warping window in {0} doesn't have a valid RootMotionModifier"),
-			FText::AsCultureInvariant(GetNameSafe(ContainingAsset)));
-		AssetCheckLog.Warning()
-			->AddToken(FUObjectToken::Create(ContainingAsset))
-			->AddToken(FTextToken::Create(MessageLooping));
-
-		if (GIsEditor)
-		{
-			AssetCheckLog.Notify(MessageLooping, EMessageSeverity::Warning, /*bForce=*/ true);
-		}
-	}
-	else if(const URootMotionModifier_Warp* RootMotionModifierWarp = Cast<URootMotionModifier_Warp>(RootMotionModifier))
-	{
-		if(RootMotionModifierWarp->WarpTargetName.IsNone())
+		if (RootMotionModifier == nullptr)
 		{
-			UObject* ContainingAsset = GetContainingAsset();
-
 			FMessageLog AssetCheckLog(NAME_AssetCheck);
 
 			const FText MessageLooping = FText::Format(
-				NSLOCTEXT("AnimNotify", "MotionWarping_InvalidWarpTargetName", "Motion Warping window in {0} doesn't specify a valid Warp Target Name"),
+				NSLOCTEXT("AnimNotify", "MotionWarping_InvalidRootMotionModifier", "Motion Warping window in {0} doesn't have a valid RootMotionModifier"),
 				FText::AsCultureInvariant(GetNameSafe(ContainingAsset)));
 			AssetCheckLog.Warning()
 				->AddToken(FUObjectToken::Create(ContainingAsset))
@@ -107,7 +87,28 @@ void UAnimNotifyState_MotionWarping::ValidateAssociatedAssets()
 
 			if (GIsEditor)
 			{
-				AssetCheckLog.Notify(MessageLooping, EMessageSeverity::Warning, /*bForce=*/ true);
+				const bool bForce = true;
+				AssetCheckLog.Notify(MessageLooping, EMessageSeverity::Warning, bForce);
+			}
+		}
+		else if (const URootMotionModifier_Warp* RootMotionModifierWarp = Cast<URootMotionModifier_Warp>(RootMotionModifier))
+		{
+			if (RootMotionModifierWarp->WarpTargetName.IsNone())
+			{
+				FMessageLog AssetCheckLog(NAME_AssetCheck);
+
+				const FText MessageLooping = FText::Format(
+					NSLOCTEXT("AnimNotify", "MotionWarping_InvalidWarpTargetName", "Motion Warping window in {0} doesn't specify a valid Warp Target Name"),
+					FText::AsCultureInvariant(GetNameSafe(ContainingAsset)));
+				AssetCheckLog.Warning()
+					->AddToken(FUObjectToken::Create(ContainingAsset))
+					->AddToken(FTextToken::Create(MessageLooping));
+
+				if (GIsEditor)
+				{
+					const bool bForce = true;
+					AssetCheckLog.Notify(MessageLooping, EMessageSeverity::Warning, bForce);
+				}
 			}
 		}
 	}
