Removed 'in' prefix from Add Root Motion Modifier helper functions.
Updated default value for some of the params to match the motion warping window.
Added blueprint callable helper function to add RootMotionModifierScale from blueprint.

diff --git a/Engine/Plugins/Experimental/Animation/MotionWarping/Source/MotionWarping/Private/RootMotionModifier.cpp b/Engine/Plugins/Experimental/Animation/MotionWarping/Source/MotionWarping/Private/RootMotionModifier.cpp
--- a/Engine/Plugins/Experimental/Animation/MotionWarping/Source/MotionWarping/Private/RootMotionModifier.cpp
+++ b/Engine/Plugins/Experimental/Animation/MotionWarping/Source/MotionWarping/Private/RootMotionModifier.cpp
@@ -423,4 +423,25 @@ FTransform UDEPRECATED_RootMotionModifier_SimpleWarp::ProcessRootMotion(const FT
 #endif
 
 	return FinalRootMotion;
+}
+
+// URootMotionModifier_Scale
+///////////////////////////////////////////////////////////////
+
+URootMotionModifier_Scale* URootMotionModifier_Scale::AddRootMotionModifierScale(UMotionWarpingComponent* InMotionWarpingComp, const UAnimSequenceBase* InAnimation, float InStartTime, float InEndTime, FVector InScale)
+{
+	if (ensureAlways(InMotionWarpingComp))
+	{
+		URootMotionModifier_Scale* NewModifier = NewObject<URootMotionModifier_Scale>(InMotionWarpingComp);
+		NewModifier->Animation = InAnimation;
+		NewModifier->StartTime = InStartTime;
+		NewModifier->EndTime = InEndTime;
+		NewModifier->Scale = InScale;
+
+		InMotionWarpingComp->AddModifier(NewModifier);
+
+		return NewModifier;
+	}
+
+	return nullptr;
 }
\ No newline at end of file
diff --git a/Engine/Plugins/Experimental/Animation/MotionWarping/Source/MotionWarping/Public/RootMotionModifier.h b/Engine/Plugins/Experimental/Animation/MotionWarping/Source/MotionWarping/Public/RootMotionModifier.h
--- a/Engine/Plugins/Experimental/Animation/MotionWarping/Source/MotionWarping/Public/RootMotionModifier.h
+++ b/Engine/Plugins/Experimental/Animation/MotionWarping/Source/MotionWarping/Public/RootMotionModifier.h
@@ -306,4 +306,12 @@ public:
 		FinalRootMotion.ScaleTranslation(Scale);
 		return FinalRootMotion;
 	}
+
+	UFUNCTION(BlueprintCallable, Category = "Motion Warping")
+	static URootMotionModifier_Scale* AddRootMotionModifierScale(
+		UPARAM(DisplayName = "Motion Warping Comp") UMotionWarpingComponent* InMotionWarpingComp,
+		UPARAM(DisplayName = "Animation") const UAnimSequenceBase* InAnimation,
+		UPARAM(DisplayName = "Start Time") float InStartTime,
+		UPARAM(DisplayName = "End Time") float InEndTime,
+		UPARAM(DisplayName = "Scale") FVector InScale = FVector(1.f));
 };
\ No newline at end of file
diff --git a/Engine/Plugins/Experimental/Animation/MotionWarping/Source/MotionWarping/Public/RootMotionModifier_SkewWarp.h b/Engine/Plugins/Experimental/Animation/MotionWarping/Source/MotionWarping/Public/RootMotionModifier_SkewWarp.h
--- a/Engine/Plugins/Experimental/Animation/MotionWarping/Source/MotionWarping/Public/RootMotionModifier_SkewWarp.h
+++ b/Engine/Plugins/Experimental/Animation/MotionWarping/Source/MotionWarping/Public/RootMotionModifier_SkewWarp.h
@@ -18,7 +18,18 @@ public:
 	virtual FTransform ProcessRootMotion(const FTransform& InRootMotion, float DeltaSeconds) override;
 
 	UFUNCTION(BlueprintCallable, Category = "Motion Warping")
-	static URootMotionModifier_SkewWarp* AddRootMotionModifierSkewWarp(UMotionWarpingComponent* InMotionWarpingComp, const UAnimSequenceBase* InAnimation, float InStartTime, float InEndTime,
-		FName InWarpTargetName, EWarpPointAnimProvider InWarpPointAnimProvider, FTransform InWarpPointAnimTransform, FName InWarpPointAnimBoneName,
-		bool bInWarpTranslation, bool bInIgnoreZAxis, bool bInWarpRotation, EMotionWarpRotationType InRotationType, float InWarpRotationTimeMultiplier = 1.f);
+	static URootMotionModifier_SkewWarp* AddRootMotionModifierSkewWarp(
+		UPARAM(DisplayName = "Motion Warping Comp") UMotionWarpingComponent* InMotionWarpingComp,
+		UPARAM(DisplayName = "Animation") const UAnimSequenceBase* InAnimation,
+		UPARAM(DisplayName = "Start Time") float InStartTime,
+		UPARAM(DisplayName = "End Time") float InEndTime,
+		UPARAM(DisplayName = "Warp Target Name") FName InWarpTargetName,
+		UPARAM(DisplayName = "Warp Point Anim Provider") EWarpPointAnimProvider InWarpPointAnimProvider,
+		UPARAM(DisplayName = "Warp Point Anim Transform") FTransform InWarpPointAnimTransform,
+		UPARAM(DisplayName = "Warp Point Anim Bone Name") FName InWarpPointAnimBoneName,
+		UPARAM(DisplayName = "Warp Translation") bool bInWarpTranslation = true,
+		UPARAM(DisplayName = "Ignore Z Axis") bool bInIgnoreZAxis = true,
+		UPARAM(DisplayName = "Warp Rotation") bool bInWarpRotation = true,
+		UPARAM(DisplayName = "Rotation Type") EMotionWarpRotationType InRotationType = EMotionWarpRotationType::Default,
+		UPARAM(DisplayName = "Warp Rotation Time Multiplier") float InWarpRotationTimeMultiplier = 1.f);
 };
\ No newline at end of file
