Motion Warping: Fixed bug preventing the character from reaching the target when warping an animation with rotation.

diff --git a/Engine/Plugins/Experimental/Animation/MotionWarping/Source/MotionWarping/Private/RootMotionModifier_SkewWarp.cpp b/Engine/Plugins/Experimental/Animation/MotionWarping/Source/MotionWarping/Private/RootMotionModifier_SkewWarp.cpp
--- a/Engine/Plugins/Experimental/Animation/MotionWarping/Source/MotionWarping/Private/RootMotionModifier_SkewWarp.cpp
+++ b/Engine/Plugins/Experimental/Animation/MotionWarping/Source/MotionWarping/Private/RootMotionModifier_SkewWarp.cpp
@@ -28,25 +28,18 @@ FTransform URootMotionModifier_SkewWarp::ProcessRootMotion(const FTransform& InR
 
 	if (bWarpTranslation && !RootMotionDelta.GetTranslation().IsNearlyZero())
 	{
-		const FTransform CurrentTransform = FTransform(
-			CharacterOwner->GetActorQuat(),
-			CharacterOwner->GetActorLocation() - FVector(0.f, 0.f, CharacterOwner->GetCapsuleComponent()->GetScaledCapsuleHalfHeight())
-		);
+		const float CapsuleHalfHeight = CharacterOwner->GetCapsuleComponent()->GetScaledCapsuleHalfHeight();
+		const FQuat CurrentRotation = CharacterOwner->GetActorQuat();
+		const FVector CurrentLocation = (CharacterOwner->GetActorLocation() - CurrentRotation.GetUpVector() * CapsuleHalfHeight);
 
-		const FTransform RootMotionTotalWorldSpace = CurrentTransform * CharacterOwner->GetMesh()->ConvertLocalRootMotionToWorld(RootMotionTotal);
-		const FTransform RootMotionDeltaWorldSpace = CharacterOwner->GetMesh()->ConvertLocalRootMotionToWorld(RootMotionDelta);
-
-		const FVector CurrentLocation = CurrentTransform.GetLocation();
-		const FQuat CurrentRotation = CurrentTransform.GetRotation();
-		
 		FVector TargetLocation = GetTargetLocation();
-		if(bIgnoreZAxis)
+		if (bIgnoreZAxis)
 		{
 			TargetLocation.Z = CurrentLocation.Z;
 		}
 
-		const FVector Translation = RootMotionDeltaWorldSpace.GetTranslation();
-		const FVector FutureLocation = RootMotionTotalWorldSpace.GetLocation();
+		const FVector Translation = CharacterOwner->GetMesh()->ConvertLocalRootMotionToWorld(RootMotionDelta).GetLocation();
+		const FVector FutureLocation = CurrentLocation + CharacterOwner->GetMesh()->ConvertLocalRootMotionToWorld(RootMotionTotal).GetLocation();
 		const FVector CurrentToWorldOffset = TargetLocation - CurrentLocation;
 		const FVector CurrentToRootOffset = FutureLocation - CurrentLocation;
 
