MotionWarping: Updated ExtractRootMotionFromAnimation so it works regardless of the value of bEnableRootMotion flag. Moved warping core algorithm and logic to calculate root relative transform from warp point to helper functions.

diff --git a/Engine/Plugins/Experimental/Animation/MotionWarping/Source/MotionWarping/Private/MotionWarpingComponent.cpp b/Engine/Plugins/Experimental/Animation/MotionWarping/Source/MotionWarping/Private/MotionWarpingComponent.cpp
--- a/Engine/Plugins/Experimental/Animation/MotionWarping/Source/MotionWarping/Private/MotionWarpingComponent.cpp
+++ b/Engine/Plugins/Experimental/Animation/MotionWarping/Source/MotionWarping/Private/MotionWarpingComponent.cpp
@@ -56,7 +56,28 @@ FTransform UMotionWarpingUtilities::ExtractRootMotionFromAnimation(const UAnimSe
 {
 	if (const UAnimMontage* Anim = Cast<UAnimMontage>(Animation))
 	{
-		return Anim->ExtractRootMotionFromTrackRange(StartTime, EndTime);
+		// This is identical to UAnimMontage::ExtractRootMotionFromTrackRange and UAnimCompositeBase::ExtractRootMotionFromTrack but ignoring bEnableRootMotion
+		// so we can extract root motion from the montage even if that flag is set to false in the AnimSequence(s)
+
+		FRootMotionMovementParams AccumulatedRootMotionParams;
+
+		if (Anim->SlotAnimTracks.Num() > 0)
+		{
+			const FAnimTrack& RootMotionAnimTrack = Anim->SlotAnimTracks[0].AnimTrack;
+
+			TArray<FRootMotionExtractionStep> RootMotionExtractionSteps;
+			RootMotionAnimTrack.GetRootMotionExtractionStepsForTrackRange(RootMotionExtractionSteps, StartTime, EndTime);
+
+			for (const FRootMotionExtractionStep& CurStep : RootMotionExtractionSteps)
+			{
+				if (CurStep.AnimSequence)
+				{
+					AccumulatedRootMotionParams.Accumulate(CurStep.AnimSequence->ExtractRootMotionFromRange(CurStep.StartPosition, CurStep.EndPosition));
+				}
+			}
+		}
+
+		return AccumulatedRootMotionParams.GetRootMotionTransform();
 	}
 
 	if (const UAnimSequence* Anim = Cast<UAnimSequence>(Animation))
@@ -136,6 +157,45 @@ void UMotionWarpingUtilities::GetMotionWarpingWindowsForWarpTargetFromAnimation(
 	}
 }
 
+FTransform UMotionWarpingUtilities::CalculateRootTransformRelativeToWarpPointAtTime(const ACharacter& Character, const UAnimSequenceBase* Animation, float Time, const FName& WarpPointBoneName)
+{
+	if (const USkeletalMeshComponent* Mesh = Character.GetMesh())
+	{
+		if (const UAnimInstance* AnimInstance = Mesh->GetAnimInstance())
+		{
+			const FBoneContainer& FullBoneContainer = AnimInstance->GetRequiredBones();
+			const int32 BoneIndex = FullBoneContainer.GetPoseBoneIndexForBoneName(WarpPointBoneName);
+			if (ensure(BoneIndex != INDEX_NONE))
+			{
+				TArray<FBoneIndexType> RequiredBoneIndexArray = { 0, (FBoneIndexType)BoneIndex };
+				FullBoneContainer.GetReferenceSkeleton().EnsureParentsExistAndSort(RequiredBoneIndexArray);
+
+				FBoneContainer LimitedBoneContainer(RequiredBoneIndexArray, FCurveEvaluationOption(false), *FullBoneContainer.GetAsset());
+
+				FCSPose<FCompactPose> Pose;
+				UMotionWarpingUtilities::ExtractComponentSpacePose(Animation, LimitedBoneContainer, Time, false, Pose);
+
+				// Inverse of mesh's relative rotation. Used to convert root and warp point in the animation from Y forward to X forward
+				const FTransform MeshCompRelativeRotInverse = FTransform(Character.GetBaseRotationOffset().Inverse());
+
+				const FTransform RootTransform = MeshCompRelativeRotInverse * Pose.GetComponentSpaceTransform(FCompactPoseBoneIndex(0));
+				const FTransform WarpPointTransform = MeshCompRelativeRotInverse * Pose.GetComponentSpaceTransform(FCompactPoseBoneIndex(1));
+				return RootTransform.GetRelativeTransform(WarpPointTransform);
+			}
+		}
+	}
+
+	return FTransform::Identity;
+}
+
+FTransform UMotionWarpingUtilities::CalculateRootTransformRelativeToWarpPointAtTime(const ACharacter& Character, const UAnimSequenceBase* Animation, float Time, const FTransform& WarpPointTransform)
+{
+	// Inverse of mesh's relative rotation. Used to convert root and warp point in the animation from Y forward to X forward
+	const FTransform MeshCompRelativeRotInverse = FTransform(Character.GetBaseRotationOffset().Inverse());
+	const FTransform RootTransform = MeshCompRelativeRotInverse * UMotionWarpingUtilities::ExtractRootTransformFromAnimation(Animation, Time);
+	return RootTransform.GetRelativeTransform((MeshCompRelativeRotInverse * WarpPointTransform));
+}
+
 // UMotionWarpingComponent
 ///////////////////////////////////////////////////////////////////////
 
diff --git a/Engine/Plugins/Experimental/Animation/MotionWarping/Source/MotionWarping/Private/RootMotionModifier.cpp b/Engine/Plugins/Experimental/Animation/MotionWarping/Source/MotionWarping/Private/RootMotionModifier.cpp
--- a/Engine/Plugins/Experimental/Animation/MotionWarping/Source/MotionWarping/Private/RootMotionModifier.cpp
+++ b/Engine/Plugins/Experimental/Animation/MotionWarping/Source/MotionWarping/Private/RootMotionModifier.cpp
@@ -218,43 +218,14 @@ void URootMotionModifier_Warp::Update()
 			{
 				if (const ACharacter* CharacterOwner = GetCharacterOwner())
 				{
-					// Inverse of mesh's relative rotation. Used to convert root and warp point in the animation from Y forward to X forward
-					const FTransform MeshCompRelativeRotInverse = FTransform(CharacterOwner->GetBaseRotationOffset().Inverse());
-
-					FTransform RootTransform = FTransform::Identity;
-					FTransform WarpPointTransform = FTransform::Identity;
 					if (WarpPointAnimProvider == EWarpPointAnimProvider::Static)
 					{
-						RootTransform = MeshCompRelativeRotInverse * UMotionWarpingUtilities::ExtractRootTransformFromAnimation(GetAnimation(), EndTime);
-						WarpPointTransform = MeshCompRelativeRotInverse * WarpPointAnimTransform;
+						CachedOffsetFromWarpPoint = UMotionWarpingUtilities::CalculateRootTransformRelativeToWarpPointAtTime(*CharacterOwner, GetAnimation(), EndTime, WarpPointAnimTransform);
 					}
 					else if (WarpPointAnimProvider == EWarpPointAnimProvider::Bone)
 					{
-						if (const USkeletalMeshComponent* Mesh = CharacterOwner->GetMesh())
-						{
-							if (const UAnimInstance* AnimInstance = Mesh->GetAnimInstance())
-							{
-								const FBoneContainer& FullBoneContainer = AnimInstance->GetRequiredBones();
-								const int32 BoneIndex = FullBoneContainer.GetPoseBoneIndexForBoneName(WarpPointAnimBoneName);
-								if (ensure(BoneIndex != INDEX_NONE))
-								{
-									TArray<FBoneIndexType> RequiredBoneIndexArray = { 0, (FBoneIndexType)BoneIndex };
-									FullBoneContainer.GetReferenceSkeleton().EnsureParentsExistAndSort(RequiredBoneIndexArray);
-
-									FBoneContainer LimitedBoneContainer(RequiredBoneIndexArray, FCurveEvaluationOption(false), *FullBoneContainer.GetAsset());
-
-									FCSPose<FCompactPose> Pose;
-									UMotionWarpingUtilities::ExtractComponentSpacePose(GetAnimation(), LimitedBoneContainer, EndTime, false, Pose);
-
-									RootTransform = MeshCompRelativeRotInverse * Pose.GetComponentSpaceTransform(FCompactPoseBoneIndex(0));
-									WarpPointTransform = MeshCompRelativeRotInverse * Pose.GetComponentSpaceTransform(FCompactPoseBoneIndex(1));
-								}
-							}
-						}
+						CachedOffsetFromWarpPoint = UMotionWarpingUtilities::CalculateRootTransformRelativeToWarpPointAtTime(*CharacterOwner, GetAnimation(), EndTime, WarpPointAnimBoneName);
 					}
-
-					// Calculate and cache the offset between the root and the warp point in the animation
-					CachedOffsetFromWarpPoint = RootTransform.GetRelativeTransform(WarpPointTransform);
 				}
 			}
 
diff --git a/Engine/Plugins/Experimental/Animation/MotionWarping/Source/MotionWarping/Private/RootMotionModifier_SkewWarp.cpp b/Engine/Plugins/Experimental/Animation/MotionWarping/Source/MotionWarping/Private/RootMotionModifier_SkewWarp.cpp
--- a/Engine/Plugins/Experimental/Animation/MotionWarping/Source/MotionWarping/Private/RootMotionModifier_SkewWarp.cpp
+++ b/Engine/Plugins/Experimental/Animation/MotionWarping/Source/MotionWarping/Private/RootMotionModifier_SkewWarp.cpp
@@ -13,33 +13,13 @@ URootMotionModifier_SkewWarp::URootMotionModifier_SkewWarp(const FObjectInitiali
 {
 }
 
-FTransform URootMotionModifier_SkewWarp::ProcessRootMotion(const FTransform& InRootMotion, float DeltaSeconds)
+FVector URootMotionModifier_SkewWarp::WarpTranslation(const FTransform& CurrentTransform, const FVector& DeltaTranslation, const FVector& TotalTranslation, const FVector& TargetLocation)
 {
-	const ACharacter* CharacterOwner = GetCharacterOwner();
-	if(CharacterOwner == nullptr)
-	{
-		return InRootMotion;
-	}
-
-	FTransform FinalRootMotion = InRootMotion;
-
-	const FTransform RootMotionTotal = UMotionWarpingUtilities::ExtractRootMotionFromAnimation(Animation.Get(), PreviousPosition, EndTime);
-	const FTransform RootMotionDelta = UMotionWarpingUtilities::ExtractRootMotionFromAnimation(Animation.Get(), PreviousPosition, FMath::Min(CurrentPosition, EndTime));
-
-	if (bWarpTranslation && !RootMotionDelta.GetTranslation().IsNearlyZero())
+	if (!DeltaTranslation.IsNearlyZero())
 	{
-		const float CapsuleHalfHeight = CharacterOwner->GetCapsuleComponent()->GetScaledCapsuleHalfHeight();
-		const FQuat CurrentRotation = CharacterOwner->GetActorQuat();
-		const FVector CurrentLocation = (CharacterOwner->GetActorLocation() - CurrentRotation.GetUpVector() * CapsuleHalfHeight);
-
-		FVector TargetLocation = GetTargetLocation();
-		if (bIgnoreZAxis)
-		{
-			TargetLocation.Z = CurrentLocation.Z;
-		}
-
-		const FVector Translation = CharacterOwner->GetMesh()->ConvertLocalRootMotionToWorld(RootMotionDelta).GetLocation();
-		const FVector FutureLocation = CurrentLocation + CharacterOwner->GetMesh()->ConvertLocalRootMotionToWorld(RootMotionTotal).GetLocation();
+		const FQuat CurrentRotation = CurrentTransform.GetRotation();
+		const FVector CurrentLocation = CurrentTransform.GetLocation();
+		const FVector FutureLocation = CurrentLocation + TotalTranslation;
 		const FVector CurrentToWorldOffset = TargetLocation - CurrentLocation;
 		const FVector CurrentToRootOffset = FutureLocation - CurrentLocation;
 
@@ -63,7 +43,7 @@ FTransform URootMotionModifier_SkewWarp::ProcessRootMotion(const FTransform& InR
 		}
 
 		// Put everything into RootSyncSpace.
-		const FVector RootMotionInSyncSpace = ToRootSyncSpace.InverseTransformVector(Translation);
+		const FVector RootMotionInSyncSpace = ToRootSyncSpace.InverseTransformVector(DeltaTranslation);
 		const FVector CurrentToWorldSync = ToRootSyncSpace.InverseTransformVector(CurrentToWorldOffset);
 		const FVector CurrentToRootMotionSync = ToRootSyncSpace.InverseTransformVector(CurrentToRootOffset);
 
@@ -133,7 +113,44 @@ FTransform URootMotionModifier_SkewWarp::ProcessRootMotion(const FTransform& InR
 		}
 
 		// Put our result back in world space.  
-		FinalRootMotion.SetTranslation(ToRootSyncSpace.TransformVector(SkewedRootMotion));
+		return ToRootSyncSpace.TransformVector(SkewedRootMotion);
+	}
+
+	return FVector::ZeroVector;
+}
+
+FTransform URootMotionModifier_SkewWarp::ProcessRootMotion(const FTransform& InRootMotion, float DeltaSeconds)
+{
+	const ACharacter* CharacterOwner = GetCharacterOwner();
+	if(CharacterOwner == nullptr)
+	{
+		return InRootMotion;
+	}
+
+	FTransform FinalRootMotion = InRootMotion;
+
+	const FTransform RootMotionTotal = UMotionWarpingUtilities::ExtractRootMotionFromAnimation(Animation.Get(), PreviousPosition, EndTime);
+	const FTransform RootMotionDelta = UMotionWarpingUtilities::ExtractRootMotionFromAnimation(Animation.Get(), PreviousPosition, FMath::Min(CurrentPosition, EndTime));
+
+	if (bWarpTranslation && !RootMotionDelta.GetTranslation().IsNearlyZero())
+	{
+		const float CapsuleHalfHeight = CharacterOwner->GetCapsuleComponent()->GetScaledCapsuleHalfHeight();
+		const FQuat CurrentRotation = CharacterOwner->GetActorQuat();
+		const FVector CurrentLocation = (CharacterOwner->GetActorLocation() - CurrentRotation.GetUpVector() * CapsuleHalfHeight);
+		const FTransform CurrentTransform = FTransform(CurrentRotation, CurrentLocation);
+
+		const FVector DeltaTranslation = CharacterOwner->GetMesh()->ConvertLocalRootMotionToWorld(RootMotionDelta).GetLocation();
+		const FVector TotalTranslation = CharacterOwner->GetMesh()->ConvertLocalRootMotionToWorld(RootMotionTotal).GetLocation();
+
+		FVector TargetLocation = GetTargetLocation();
+		if (bIgnoreZAxis)
+		{
+			TargetLocation.Z = CurrentLocation.Z;
+		}
+
+		const FVector WarpedTranslation = WarpTranslation(CurrentTransform, DeltaTranslation, TotalTranslation, TargetLocation);
+
+		FinalRootMotion.SetTranslation(WarpedTranslation);
 	}
 
 	if(bWarpRotation)
diff --git a/Engine/Plugins/Experimental/Animation/MotionWarping/Source/MotionWarping/Public/MotionWarpingComponent.h b/Engine/Plugins/Experimental/Animation/MotionWarping/Source/MotionWarping/Public/MotionWarpingComponent.h
--- a/Engine/Plugins/Experimental/Animation/MotionWarping/Source/MotionWarping/Public/MotionWarpingComponent.h
+++ b/Engine/Plugins/Experimental/Animation/MotionWarping/Source/MotionWarping/Public/MotionWarpingComponent.h
@@ -67,6 +67,12 @@ public:
 	/** @return All the MotionWarping windows within the supplied animation for a given Warp Target */
 	UFUNCTION(BlueprintCallable, Category = "Motion Warping")
 	static void GetMotionWarpingWindowsForWarpTargetFromAnimation(const UAnimSequenceBase* Animation, FName WarpTargetName, TArray<FMotionWarpingWindowData>& OutWindows);
+
+	/** @return root transform relative to the warp point bone at the supplied time */
+	static FTransform CalculateRootTransformRelativeToWarpPointAtTime(const ACharacter& Character, const UAnimSequenceBase* Animation, float Time, const FName& WarpPointBoneName);
+
+	/** @return root transform relative to the warp point transform at the supplied time */
+	static FTransform CalculateRootTransformRelativeToWarpPointAtTime(const ACharacter& Character, const UAnimSequenceBase* Animation, float Time, const FTransform& WarpPointTransform);
 };
 
 DECLARE_DYNAMIC_MULTICAST_DELEGATE_OneParam(FMotionWarpingPreUpdate, class UMotionWarpingComponent*, MotionWarpingComp);
diff --git a/Engine/Plugins/Experimental/Animation/MotionWarping/Source/MotionWarping/Public/RootMotionModifier_SkewWarp.h b/Engine/Plugins/Experimental/Animation/MotionWarping/Source/MotionWarping/Public/RootMotionModifier_SkewWarp.h
--- a/Engine/Plugins/Experimental/Animation/MotionWarping/Source/MotionWarping/Public/RootMotionModifier_SkewWarp.h
+++ b/Engine/Plugins/Experimental/Animation/MotionWarping/Source/MotionWarping/Public/RootMotionModifier_SkewWarp.h
@@ -17,6 +17,8 @@ public:
 
 	virtual FTransform ProcessRootMotion(const FTransform& InRootMotion, float DeltaSeconds) override;
 
+	static FVector WarpTranslation(const FTransform& CurrentTransform, const FVector& DeltaTranslation, const FVector& TotalTranslation, const FVector& TargetLocation);
+
 	UFUNCTION(BlueprintCallable, Category = "Motion Warping")
 	static URootMotionModifier_SkewWarp* AddRootMotionModifierSkewWarp(
 		UPARAM(DisplayName = "Motion Warping Comp") UMotionWarpingComponent* InMotionWarpingComp,
