Fix for static analysis warning

diff --git a/Engine/Plugins/Experimental/Animation/MotionWarping/Source/MotionWarping/Private/MotionWarpingComponent.cpp b/Engine/Plugins/Experimental/Animation/MotionWarping/Source/MotionWarping/Private/MotionWarpingComponent.cpp
--- a/Engine/Plugins/Experimental/Animation/MotionWarping/Source/MotionWarping/Private/MotionWarpingComponent.cpp
+++ b/Engine/Plugins/Experimental/Animation/MotionWarping/Source/MotionWarping/Private/MotionWarpingComponent.cpp
@@ -294,33 +294,35 @@ void UMotionWarpingComponent::Update()
 			for (int32 SlotIdx = 0; SlotIdx < Montage->SlotAnimTracks.Num(); SlotIdx++)
 			{
 				const FAnimTrack& AnimTrack = Montage->SlotAnimTracks[SlotIdx].AnimTrack;
-				const FAnimSegment* AnimSegment = AnimTrack.GetSegmentAtTime(PreviousPosition);
-				const UAnimSequenceBase* AnimReference = AnimSegment->GetAnimReference();
-				if (AnimSegment && AnimReference)
+
+				if (const FAnimSegment* AnimSegment = AnimTrack.GetSegmentAtTime(PreviousPosition))
 				{
-					for (const FAnimNotifyEvent& NotifyEvent : AnimReference->Notifies)
+					if(const UAnimSequenceBase* AnimReference = AnimSegment->GetAnimReference())
 					{
-						const UAnimNotifyState_MotionWarping* MotionWarpingNotify = NotifyEvent.NotifyStateClass ? Cast<UAnimNotifyState_MotionWarping>(NotifyEvent.NotifyStateClass) : nullptr;
-						if (MotionWarpingNotify)
+						for (const FAnimNotifyEvent& NotifyEvent : AnimReference->Notifies)
 						{
-							if (MotionWarpingNotify->RootMotionModifier == nullptr)
+							const UAnimNotifyState_MotionWarping* MotionWarpingNotify = NotifyEvent.NotifyStateClass ? Cast<UAnimNotifyState_MotionWarping>(NotifyEvent.NotifyStateClass) : nullptr;
+							if (MotionWarpingNotify)
 							{
-								UE_LOG(LogMotionWarping, Warning, TEXT("MotionWarpingComponent::Update. A motion warping window in %s doesn't have a valid root motion modifier!"), *GetNameSafe(AnimReference));
-								continue;
-							}
+								if (MotionWarpingNotify->RootMotionModifier == nullptr)
+								{
+									UE_LOG(LogMotionWarping, Warning, TEXT("MotionWarpingComponent::Update. A motion warping window in %s doesn't have a valid root motion modifier!"), *GetNameSafe(AnimReference));
+									continue;
+								}
 
-							const float NotifyStartTime = FMath::Clamp(NotifyEvent.GetTriggerTime(), 0.f, AnimReference->GetPlayLength());
-							const float NotifyEndTime = FMath::Clamp(NotifyEvent.GetEndTriggerTime(), 0.f, AnimReference->GetPlayLength());
+								const float NotifyStartTime = FMath::Clamp(NotifyEvent.GetTriggerTime(), 0.f, AnimReference->GetPlayLength());
+								const float NotifyEndTime = FMath::Clamp(NotifyEvent.GetEndTriggerTime(), 0.f, AnimReference->GetPlayLength());
 
-							// Convert notify times from AnimSequence times to montage times
-							const float StartTime = (NotifyStartTime - AnimSegment->AnimStartTime) + AnimSegment->StartPos;
-							const float EndTime = (NotifyEndTime - AnimSegment->AnimStartTime) + AnimSegment->StartPos;
+								// Convert notify times from AnimSequence times to montage times
+								const float StartTime = (NotifyStartTime - AnimSegment->AnimStartTime) + AnimSegment->StartPos;
+								const float EndTime = (NotifyEndTime - AnimSegment->AnimStartTime) + AnimSegment->StartPos;
 
-							if (PreviousPosition >= StartTime && PreviousPosition < EndTime)
-							{
-								if (!ContainsModifier(Montage, StartTime, EndTime))
+								if (PreviousPosition >= StartTime && PreviousPosition < EndTime)
 								{
-									MotionWarpingNotify->OnBecomeRelevant(this, Montage, StartTime, EndTime);
+									if (!ContainsModifier(Montage, StartTime, EndTime))
+									{
+										MotionWarpingNotify->OnBecomeRelevant(this, Montage, StartTime, EndTime);
+									}
 								}
 							}
 						}
