Motion Warping: Fixed bug where root motion outside warp window was getting stomped during the last tick

diff --git a/Engine/Plugins/Experimental/Animation/MotionWarping/Source/MotionWarping/Private/RootMotionModifier_SkewWarp.cpp b/Engine/Plugins/Experimental/Animation/MotionWarping/Source/MotionWarping/Private/RootMotionModifier_SkewWarp.cpp
--- a/Engine/Plugins/Experimental/Animation/MotionWarping/Source/MotionWarping/Private/RootMotionModifier_SkewWarp.cpp
+++ b/Engine/Plugins/Experimental/Animation/MotionWarping/Source/MotionWarping/Private/RootMotionModifier_SkewWarp.cpp
@@ -132,6 +132,13 @@ FTransform URootMotionModifier_SkewWarp::ProcessRootMotion(const FTransform& InR
 	const FTransform RootMotionTotal = UMotionWarpingUtilities::ExtractRootMotionFromAnimation(Animation.Get(), PreviousPosition, EndTime);
 	const FTransform RootMotionDelta = UMotionWarpingUtilities::ExtractRootMotionFromAnimation(Animation.Get(), PreviousPosition, FMath::Min(CurrentPosition, EndTime));
 
+	FTransform ExtraRootMotion = FTransform::Identity;
+	if (CurrentPosition > EndTime)
+	{
+		ExtraRootMotion = UMotionWarpingUtilities::ExtractRootMotionFromAnimation(Animation.Get(), EndTime, CurrentPosition);
+		ExtraRootMotion = CharacterOwner->GetMesh()->ConvertLocalRootMotionToWorld(ExtraRootMotion);
+	}
+
 	if (bWarpTranslation && !RootMotionDelta.GetTranslation().IsNearlyZero())
 	{
 		const float CapsuleHalfHeight = CharacterOwner->GetCapsuleComponent()->GetScaledCapsuleHalfHeight();
@@ -148,14 +155,13 @@ FTransform URootMotionModifier_SkewWarp::ProcessRootMotion(const FTransform& InR
 			TargetLocation.Z = CurrentLocation.Z;
 		}
 
-		const FVector WarpedTranslation = WarpTranslation(CurrentTransform, DeltaTranslation, TotalTranslation, TargetLocation);
-
+		const FVector WarpedTranslation = WarpTranslation(CurrentTransform, DeltaTranslation, TotalTranslation, TargetLocation) + ExtraRootMotion.GetLocation();
 		FinalRootMotion.SetTranslation(WarpedTranslation);
 	}
 
 	if(bWarpRotation)
 	{
-		const FQuat WarpedRotation = WarpRotation(InRootMotion, RootMotionTotal, DeltaSeconds);
+		const FQuat WarpedRotation = ExtraRootMotion.GetRotation() * WarpRotation(RootMotionDelta, RootMotionTotal, DeltaSeconds);
 		FinalRootMotion.SetRotation(WarpedRotation);
 	}
 
