Added support for warping resimulated moves after a correction in multiplayer environment

diff --git a/Engine/Plugins/Experimental/Animation/MotionWarping/Source/MotionWarping/Private/MotionWarpingComponent.cpp b/Engine/Plugins/Experimental/Animation/MotionWarping/Source/MotionWarping/Private/MotionWarpingComponent.cpp
--- a/Engine/Plugins/Experimental/Animation/MotionWarping/Source/MotionWarping/Private/MotionWarpingComponent.cpp
+++ b/Engine/Plugins/Experimental/Animation/MotionWarping/Source/MotionWarping/Private/MotionWarpingComponent.cpp
@@ -254,35 +254,67 @@ void UMotionWarpingComponent::DisableAllRootMotionModifiers()
 
 void UMotionWarpingComponent::Update()
 {
-	check(GetCharacterOwner());
+	const ACharacter* Character = GetCharacterOwner();
+	check(Character);
 
-	const FAnimMontageInstance* RootMotionMontageInstance = GetCharacterOwner()->GetRootMotionAnimMontageInstance();
-	UAnimMontage* Montage = RootMotionMontageInstance ? ToRawPtr(RootMotionMontageInstance->Montage) : nullptr;
-	if (Montage)
+	FMotionWarpingUpdateContext Context;
+
+	// When replaying saved moves we need to look at the contributor to root motion back then.
+	if (Character->bClientUpdating)
+	{
+		const UCharacterMovementComponent* MoveComp = Character->GetCharacterMovement();
+		check(MoveComp);
+
+		const FSavedMove_Character* SavedMove = MoveComp->GetCurrentReplayedSavedMove();
+		check(SavedMove);
+
+		if(SavedMove->RootMotionMontage.IsValid())
+		{
+			Context.Animation = SavedMove->RootMotionMontage.Get();
+			Context.CurrentPosition = SavedMove->RootMotionTrackPosition;
+			Context.PreviousPosition = FMath::Max(Context.CurrentPosition - SavedMove->DeltaTime * SavedMove->RootMotionMontage->RateScale, 0.f);
+		}
+	}
+	else // If we are not replaying a move, just use the current root motion montage
 	{
-		const float PreviousPosition = RootMotionMontageInstance->GetPreviousPosition();
-		const float CurrentPosition = RootMotionMontageInstance->GetPosition();
+		if(const FAnimMontageInstance* RootMotionMontageInstance = Character->GetRootMotionAnimMontageInstance())
+		{
+			const UAnimMontage* Montage = RootMotionMontageInstance->Montage;
+			check(Montage);
+
+			Context.Animation = Montage;
+			Context.CurrentPosition = RootMotionMontageInstance->GetPosition();
+			Context.PreviousPosition = RootMotionMontageInstance->GetPreviousPosition();
+			Context.Weight = RootMotionMontageInstance->GetWeight();
+		}
+	}
+
+	if (Context.Animation.IsValid())
+	{
+		const UAnimSequenceBase* Animation = Context.Animation.Get();
+		const float PreviousPosition = Context.PreviousPosition;
+		const float CurrentPosition = Context.CurrentPosition;
 
 		// Loop over notifies directly in the montage, looking for Motion Warping windows
-		for (const FAnimNotifyEvent& NotifyEvent : Montage->Notifies)
+		for (const FAnimNotifyEvent& NotifyEvent : Animation->Notifies)
 		{
 			const UAnimNotifyState_MotionWarping* MotionWarpingNotify = NotifyEvent.NotifyStateClass ? Cast<UAnimNotifyState_MotionWarping>(NotifyEvent.NotifyStateClass) : nullptr;
 			if (MotionWarpingNotify)
 			{
 				if(MotionWarpingNotify->RootMotionModifier == nullptr)
 				{
-					UE_LOG(LogMotionWarping, Warning, TEXT("MotionWarpingComponent::Update. A motion warping window in %s doesn't have a valid root motion modifier!"), *GetNameSafe(Montage));
+					UE_LOG(LogMotionWarping, Warning, TEXT("MotionWarpingComponent::Update. A motion warping window in %s doesn't have a valid root motion modifier!"), *GetNameSafe(Animation));
 					continue;
 				}
 
-				const float StartTime = FMath::Clamp(NotifyEvent.GetTriggerTime(), 0.f, Montage->GetPlayLength());
-				const float EndTime = FMath::Clamp(NotifyEvent.GetEndTriggerTime(), 0.f, Montage->GetPlayLength());
+				const float StartTime = FMath::Clamp(NotifyEvent.GetTriggerTime(), 0.f, Animation->GetPlayLength());
+				const float EndTime = FMath::Clamp(NotifyEvent.GetEndTriggerTime(), 0.f, Animation->GetPlayLength());
 
 				if (PreviousPosition >= StartTime && PreviousPosition < EndTime)
 				{
-					if (!ContainsModifier(Montage, StartTime, EndTime))
+					if (!ContainsModifier(Animation, StartTime, EndTime))
 					{
-						MotionWarpingNotify->OnBecomeRelevant(this, Montage, StartTime, EndTime);
+						MotionWarpingNotify->OnBecomeRelevant(this, Animation, StartTime, EndTime);
 					}
 				}
 			}
@@ -290,38 +322,41 @@ void UMotionWarpingComponent::Update()
 
 		if(bSearchForWindowsInAnimsWithinMontages)
 		{
-			// Same as before but scanning all animation within the montage
-			for (int32 SlotIdx = 0; SlotIdx < Montage->SlotAnimTracks.Num(); SlotIdx++)
+			if(const UAnimMontage* Montage = Cast<const UAnimMontage>(Context.Animation.Get()))
 			{
-				const FAnimTrack& AnimTrack = Montage->SlotAnimTracks[SlotIdx].AnimTrack;
-
-				if (const FAnimSegment* AnimSegment = AnimTrack.GetSegmentAtTime(PreviousPosition))
+				// Same as before but scanning all animation within the montage
+				for (int32 SlotIdx = 0; SlotIdx < Montage->SlotAnimTracks.Num(); SlotIdx++)
 				{
-					if(const UAnimSequenceBase* AnimReference = AnimSegment->GetAnimReference())
+					const FAnimTrack& AnimTrack = Montage->SlotAnimTracks[SlotIdx].AnimTrack;
+
+					if (const FAnimSegment* AnimSegment = AnimTrack.GetSegmentAtTime(PreviousPosition))
 					{
-						for (const FAnimNotifyEvent& NotifyEvent : AnimReference->Notifies)
+						if (const UAnimSequenceBase* AnimReference = AnimSegment->GetAnimReference())
 						{
-							const UAnimNotifyState_MotionWarping* MotionWarpingNotify = NotifyEvent.NotifyStateClass ? Cast<UAnimNotifyState_MotionWarping>(NotifyEvent.NotifyStateClass) : nullptr;
-							if (MotionWarpingNotify)
+							for (const FAnimNotifyEvent& NotifyEvent : AnimReference->Notifies)
 							{
-								if (MotionWarpingNotify->RootMotionModifier == nullptr)
+								const UAnimNotifyState_MotionWarping* MotionWarpingNotify = NotifyEvent.NotifyStateClass ? Cast<UAnimNotifyState_MotionWarping>(NotifyEvent.NotifyStateClass) : nullptr;
+								if (MotionWarpingNotify)
 								{
-									UE_LOG(LogMotionWarping, Warning, TEXT("MotionWarpingComponent::Update. A motion warping window in %s doesn't have a valid root motion modifier!"), *GetNameSafe(AnimReference));
-									continue;
-								}
+									if (MotionWarpingNotify->RootMotionModifier == nullptr)
+									{
+										UE_LOG(LogMotionWarping, Warning, TEXT("MotionWarpingComponent::Update. A motion warping window in %s doesn't have a valid root motion modifier!"), *GetNameSafe(AnimReference));
+										continue;
+									}
 
-								const float NotifyStartTime = FMath::Clamp(NotifyEvent.GetTriggerTime(), 0.f, AnimReference->GetPlayLength());
-								const float NotifyEndTime = FMath::Clamp(NotifyEvent.GetEndTriggerTime(), 0.f, AnimReference->GetPlayLength());
+									const float NotifyStartTime = FMath::Clamp(NotifyEvent.GetTriggerTime(), 0.f, AnimReference->GetPlayLength());
+									const float NotifyEndTime = FMath::Clamp(NotifyEvent.GetEndTriggerTime(), 0.f, AnimReference->GetPlayLength());
 
-								// Convert notify times from AnimSequence times to montage times
-								const float StartTime = (NotifyStartTime - AnimSegment->AnimStartTime) + AnimSegment->StartPos;
-								const float EndTime = (NotifyEndTime - AnimSegment->AnimStartTime) + AnimSegment->StartPos;
+									// Convert notify times from AnimSequence times to montage times
+									const float StartTime = (NotifyStartTime - AnimSegment->AnimStartTime) + AnimSegment->StartPos;
+									const float EndTime = (NotifyEndTime - AnimSegment->AnimStartTime) + AnimSegment->StartPos;
 
-								if (PreviousPosition >= StartTime && PreviousPosition < EndTime)
-								{
-									if (!ContainsModifier(Montage, StartTime, EndTime))
+									if (PreviousPosition >= StartTime && PreviousPosition < EndTime)
 									{
-										MotionWarpingNotify->OnBecomeRelevant(this, Montage, StartTime, EndTime);
+										if (!ContainsModifier(Montage, StartTime, EndTime))
+										{
+											MotionWarpingNotify->OnBecomeRelevant(this, Montage, StartTime, EndTime);
+										}
 									}
 								}
 							}
@@ -339,7 +374,7 @@ void UMotionWarpingComponent::Update()
 	{
 		for (URootMotionModifier* Modifier : Modifiers)
 		{
-			Modifier->Update();
+			Modifier->Update(Context);
 		}
 
 		// Remove the modifiers that has been marked for removal
diff --git a/Engine/Plugins/Experimental/Animation/MotionWarping/Source/MotionWarping/Private/RootMotionModifier.cpp b/Engine/Plugins/Experimental/Animation/MotionWarping/Source/MotionWarping/Private/RootMotionModifier.cpp
--- a/Engine/Plugins/Experimental/Animation/MotionWarping/Source/MotionWarping/Private/RootMotionModifier.cpp
+++ b/Engine/Plugins/Experimental/Animation/MotionWarping/Source/MotionWarping/Private/RootMotionModifier.cpp
@@ -88,7 +88,7 @@ ACharacter* URootMotionModifier::GetCharacterOwner() const
 	return OwnerComp ? OwnerComp->GetCharacterOwner() : nullptr;
 }
 
-void URootMotionModifier::Update()
+void URootMotionModifier::Update(const FMotionWarpingUpdateContext& Context)
 {
 	const ACharacter* CharacterOwner = GetCharacterOwner();
 	if (CharacterOwner == nullptr)
@@ -96,23 +96,20 @@ void URootMotionModifier::Update()
 		return;
 	}
 
-	const FAnimMontageInstance* RootMotionMontageInstance = CharacterOwner->GetRootMotionAnimMontageInstance();
-	const UAnimMontage* Montage = RootMotionMontageInstance ? ToRawPtr(RootMotionMontageInstance->Montage) : nullptr;
-
 	// Mark for removal if our animation is not relevant anymore
-	if (Montage == nullptr || Montage != Animation)
+	if (!Context.Animation.IsValid() || Context.Animation.Get() != Animation)
 	{
-		UE_LOG(LogMotionWarping, Verbose, TEXT("MotionWarping: Marking RootMotionModifier for removal. Reason: Animation is not valid. Char: %s Current Montage: %s. Window: Animation: %s [%f %f] [%f %f]"),
-			*GetNameSafe(CharacterOwner), *GetNameSafe(Montage), *GetNameSafe(Animation.Get()), StartTime, EndTime, PreviousPosition, CurrentPosition);
+		UE_LOG(LogMotionWarping, Verbose, TEXT("MotionWarping: Marking RootMotionModifier for removal. Reason: Animation is not valid. Char: %s Current Animation: %s. Window: Animation: %s [%f %f] [%f %f]"),
+			*GetNameSafe(CharacterOwner), *GetNameSafe(Context.Animation.Get()), *GetNameSafe(Animation.Get()), StartTime, EndTime, PreviousPosition, CurrentPosition);
 
 		SetState(ERootMotionModifierState::MarkedForRemoval);
 		return;
 	}
 
 	// Update playback times and weight
-	PreviousPosition = RootMotionMontageInstance->GetPreviousPosition();
-	CurrentPosition = RootMotionMontageInstance->GetPosition();
-	Weight = RootMotionMontageInstance->GetWeight();
+	PreviousPosition = Context.PreviousPosition;
+	CurrentPosition = Context.CurrentPosition;
+	Weight = Context.Weight;
 
 	// Mark for removal if the animation already passed the warping window
 	if (PreviousPosition >= EndTime)
@@ -183,10 +180,10 @@ URootMotionModifier_Warp::URootMotionModifier_Warp(const FObjectInitializer& Obj
 {
 }
 
-void URootMotionModifier_Warp::Update()
+void URootMotionModifier_Warp::Update(const FMotionWarpingUpdateContext& Context)
 {
 	// Update playback times and state
-	Super::Update();
+	Super::Update(Context);
 
 	// Cache sync point transform and trigger OnTargetTransformChanged if needed
 	const UMotionWarpingComponent* OwnerComp = GetOwnerComponent();
diff --git a/Engine/Plugins/Experimental/Animation/MotionWarping/Source/MotionWarping/Public/RootMotionModifier.h b/Engine/Plugins/Experimental/Animation/MotionWarping/Source/MotionWarping/Public/RootMotionModifier.h
--- a/Engine/Plugins/Experimental/Animation/MotionWarping/Source/MotionWarping/Public/RootMotionModifier.h
+++ b/Engine/Plugins/Experimental/Animation/MotionWarping/Source/MotionWarping/Public/RootMotionModifier.h
@@ -11,6 +11,28 @@ class UAnimNotifyState_MotionWarping;
 class URootMotionModifier;
 class USceneComponent;
 
+/** 
+ * Context passed to any active root motion modifier during the update phase. 
+ * Contains relevant data from the animation that contributed to root motion this frame (or in the past when replaying saved moves)
+ */
+USTRUCT()
+struct FMotionWarpingUpdateContext
+{
+	GENERATED_BODY()
+
+	UPROPERTY()
+	TWeakObjectPtr<const UAnimSequenceBase> Animation = nullptr;
+
+	UPROPERTY()
+	float PreviousPosition = 0.f;
+
+	UPROPERTY()
+	float CurrentPosition = 0.f;
+
+	UPROPERTY()
+	float Weight = 0.f;
+};
+
 /** The possible states of a Root Motion Modifier */
 UENUM(BlueprintType)
 enum class ERootMotionModifierState : uint8
@@ -105,7 +127,7 @@ public:
 	/** Returns a pointer to the character that owns the component that owns this modifier */
 	class ACharacter* GetCharacterOwner() const;
 
-	virtual void Update();
+	virtual void Update(const FMotionWarpingUpdateContext& Context);
 	virtual FTransform ProcessRootMotion(const FTransform& InRootMotion, float DeltaSeconds) { return FTransform::Identity; }
 
 	FORCEINLINE const UAnimSequenceBase* GetAnimation() const { return Animation.Get(); }
@@ -244,7 +266,7 @@ public:
 	URootMotionModifier_Warp(const FObjectInitializer& ObjectInitializer);
 
 	//~ Begin FRootMotionModifier Interface
-	virtual void Update() override;
+	virtual void Update(const FMotionWarpingUpdateContext& Context) override;
 	//~ End FRootMotionModifier Interface
 
 	/** Event called during update if the target transform changes while the warping is active */
