MotionWarping: Changed warp targets container from a TMap to a TArray and marked as replicated (SimulatedOnly)

diff --git a/Engine/Plugins/Experimental/Animation/MotionWarping/Source/MotionWarping/MotionWarping.Build.cs b/Engine/Plugins/Experimental/Animation/MotionWarping/Source/MotionWarping/MotionWarping.Build.cs
--- a/Engine/Plugins/Experimental/Animation/MotionWarping/Source/MotionWarping/MotionWarping.Build.cs
+++ b/Engine/Plugins/Experimental/Animation/MotionWarping/Source/MotionWarping/MotionWarping.Build.cs
@@ -28,6 +28,7 @@ public class MotionWarping : ModuleRules
 				"Core",
 				"CoreUObject",
                 "Engine",
+				"NetCore"
 			}
 			);
 			
diff --git a/Engine/Plugins/Experimental/Animation/MotionWarping/Source/MotionWarping/Private/MotionWarpingComponent.cpp b/Engine/Plugins/Experimental/Animation/MotionWarping/Source/MotionWarping/Private/MotionWarpingComponent.cpp
--- a/Engine/Plugins/Experimental/Animation/MotionWarping/Source/MotionWarping/Private/MotionWarpingComponent.cpp
+++ b/Engine/Plugins/Experimental/Animation/MotionWarping/Source/MotionWarping/Private/MotionWarpingComponent.cpp
@@ -11,6 +11,7 @@
 #include "GameFramework/Character.h"
 #include "GameFramework/CharacterMovementComponent.h"
 #include "AnimNotifyState_MotionWarping.h"
+#include "Net/UnrealNetwork.h"
 
 DEFINE_LOG_CATEGORY(LogMotionWarping);
 
@@ -203,6 +204,17 @@ UMotionWarpingComponent::UMotionWarpingComponent(const FObjectInitializer& Objec
 	: Super(ObjectInitializer)
 {
 	bWantsInitializeComponent = true;
+	SetIsReplicatedByDefault(true);
+}
+
+void UMotionWarpingComponent::GetLifetimeReplicatedProps(TArray< FLifetimeProperty >& OutLifetimeProps) const
+{
+	Super::GetLifetimeReplicatedProps(OutLifetimeProps);
+
+	FDoRepLifetimeParams Params;
+	Params.bIsPushBased = true;
+	Params.Condition = COND_SimulatedOnly;
+	DOREPLIFETIME_WITH_PARAMS_FAST(UMotionWarpingComponent, WarpTargets, Params);
 }
 
 void UMotionWarpingComponent::InitializeComponent()
@@ -474,23 +486,40 @@ FTransform UMotionWarpingComponent::ProcessRootMotionPostConvertToWorld(const FT
 	return FinalRootMotion;
 }
 
-void UMotionWarpingComponent::AddOrUpdateWarpTarget(FName WarpTargetName, const FMotionWarpingTarget& WarpTarget)
+void UMotionWarpingComponent::AddOrUpdateWarpTarget(const FMotionWarpingTarget& WarpTarget)
 {
-	if (WarpTargetName != NAME_None)
+	if (WarpTarget.Name != NAME_None)
 	{
-		FMotionWarpingTarget& WarpTargetRef = WarpTargetMap.FindOrAdd(WarpTargetName);
-		WarpTargetRef = WarpTarget;
+		for (int32 Idx = 0; Idx < WarpTargets.Num(); Idx++)
+		{
+			if (WarpTargets[Idx].Name == WarpTarget.Name)
+			{
+				WarpTargets[Idx] = WarpTarget;
+				return;
+			}
+		}
+
+		WarpTargets.Add(WarpTarget);
+
+		MARK_PROPERTY_DIRTY_FROM_NAME(UMotionWarpingComponent, WarpTargets, this);
 	}
 }
 
 int32 UMotionWarpingComponent::RemoveWarpTarget(FName WarpTargetName)
 {
-	return WarpTargetMap.Remove(WarpTargetName);
+	const int32 NumRemoved = WarpTargets.RemoveAll([&WarpTargetName](const FMotionWarpingTarget& WarpTarget) { return WarpTarget.Name == WarpTargetName; });
+	
+	if(NumRemoved > 0)
+	{
+		MARK_PROPERTY_DIRTY_FROM_NAME(UMotionWarpingComponent, WarpTargets, this);
+	}
+
+	return NumRemoved;
 }
 
 void UMotionWarpingComponent::AddOrUpdateWarpTargetFromTransform(FName WarpTargetName, FTransform TargetTransform)
 {
-	AddOrUpdateWarpTarget(WarpTargetName, TargetTransform);
+	AddOrUpdateWarpTarget(FMotionWarpingTarget(WarpTargetName, TargetTransform));
 }
 
 void UMotionWarpingComponent::AddOrUpdateWarpTargetFromComponent(FName WarpTargetName, const USceneComponent* Component, FName BoneName, bool bFollowComponent)
@@ -501,7 +530,7 @@ void UMotionWarpingComponent::AddOrUpdateWarpTargetFromComponent(FName WarpTarge
 		return;
 	}
 
-	AddOrUpdateWarpTarget(WarpTargetName, FMotionWarpingTarget(Component, BoneName, bFollowComponent));
+	AddOrUpdateWarpTarget(FMotionWarpingTarget(WarpTargetName, Component, BoneName, bFollowComponent));
 }
 
 URootMotionModifier* UMotionWarpingComponent::AddModifierFromTemplate(URootMotionModifier* Template, const UAnimSequenceBase* Animation, float StartTime, float EndTime)
diff --git a/Engine/Plugins/Experimental/Animation/MotionWarping/Source/MotionWarping/Private/RootMotionModifier.cpp b/Engine/Plugins/Experimental/Animation/MotionWarping/Source/MotionWarping/Private/RootMotionModifier.cpp
--- a/Engine/Plugins/Experimental/Animation/MotionWarping/Source/MotionWarping/Private/RootMotionModifier.cpp
+++ b/Engine/Plugins/Experimental/Animation/MotionWarping/Source/MotionWarping/Private/RootMotionModifier.cpp
@@ -14,14 +14,16 @@
 // FMotionWarpingTarget
 ///////////////////////////////////////////////////////////////
 
-FMotionWarpingTarget::FMotionWarpingTarget(const USceneComponent* InComp, FName InBoneName, bool bInFollowComponent)
+FMotionWarpingTarget::FMotionWarpingTarget(const FName& InName, const USceneComponent* InComp, FName InBoneName, bool bInFollowComponent)
 {
 	if (ensure(InComp))
 	{
+		Name = InName;
 		Component = InComp;
 		BoneName = InBoneName;
 		bFollowComponent = bInFollowComponent;
 
+		FTransform Transform = FTransform::Identity;
 		if (BoneName != NAME_None)
 		{
 			Transform = FMotionWarpingTarget::GetTargetTransformFromComponent(InComp, InBoneName);
@@ -30,6 +32,9 @@ FMotionWarpingTarget::FMotionWarpingTarget(const USceneComponent* InComp, FName
 		{
 			Transform = InComp->GetComponentTransform();
 		}
+
+		Location = Transform.GetLocation();
+		Rotation = Transform.Rotator();
 	}
 }
 
@@ -66,7 +71,7 @@ FTransform FMotionWarpingTarget::GetTargetTrasform() const
 		}
 	}
 
-	return Transform;
+	return FTransform(Rotation, Location);
 }
 
 // URootMotionModifier
diff --git a/Engine/Plugins/Experimental/Animation/MotionWarping/Source/MotionWarping/Public/MotionWarpingComponent.h b/Engine/Plugins/Experimental/Animation/MotionWarping/Source/MotionWarping/Public/MotionWarpingComponent.h
--- a/Engine/Plugins/Experimental/Animation/MotionWarping/Source/MotionWarping/Public/MotionWarpingComponent.h
+++ b/Engine/Plugins/Experimental/Animation/MotionWarping/Source/MotionWarping/Public/MotionWarpingComponent.h
@@ -95,6 +95,7 @@ public:
 	UMotionWarpingComponent(const FObjectInitializer& ObjectInitializer);
 
 	virtual void InitializeComponent() override;
+	virtual void GetLifetimeReplicatedProps(TArray< FLifetimeProperty >& OutLifetimeProps) const;
 
 	/** Gets the character this component belongs to */
 	FORCEINLINE ACharacter* GetCharacterOwner() const { return CharacterOwner.Get(); }
@@ -115,11 +116,14 @@ public:
 	URootMotionModifier* AddModifierFromTemplate(URootMotionModifier* Template, const UAnimSequenceBase* Animation, float StartTime, float EndTime);
 
 	/** Find the target associated with a specified name */
-	FORCEINLINE const FMotionWarpingTarget* FindWarpTarget(const FName& WarpTargetName) const { return WarpTargetMap.Find(WarpTargetName); }
+	FORCEINLINE const FMotionWarpingTarget* FindWarpTarget(const FName& WarpTargetName) const 
+	{ 
+		return WarpTargets.FindByPredicate([&WarpTargetName](const FMotionWarpingTarget& WarpTarget){ return WarpTarget.Name == WarpTargetName; });
+	}
 
-	/** Adds or update a target associated with a specified name */
+	/** Adds or update a warp target */
 	UFUNCTION(BlueprintCallable, Category = "Motion Warping")
-	void AddOrUpdateWarpTarget(FName WarpTargetName, const FMotionWarpingTarget& WarpTarget);
+	void AddOrUpdateWarpTarget(const FMotionWarpingTarget& WarpTarget);
 
 	/** 
 	 * Create and adds or update a target associated with a specified name 
@@ -176,8 +180,8 @@ protected:
 	UPROPERTY(Transient)
 	TArray<URootMotionModifier*> Modifiers;
 
-	UPROPERTY(Transient)
-	TMap<FName, FMotionWarpingTarget> WarpTargetMap;
+	UPROPERTY(Transient, Replicated)
+	TArray<FMotionWarpingTarget> WarpTargets;
 
 #if !(UE_BUILD_SHIPPING || UE_BUILD_TEST)
 	TOptional<FVector> OriginalRootMotionAccum;
diff --git a/Engine/Plugins/Experimental/Animation/MotionWarping/Source/MotionWarping/Public/RootMotionModifier.h b/Engine/Plugins/Experimental/Animation/MotionWarping/Source/MotionWarping/Public/RootMotionModifier.h
--- a/Engine/Plugins/Experimental/Animation/MotionWarping/Source/MotionWarping/Public/RootMotionModifier.h
+++ b/Engine/Plugins/Experimental/Animation/MotionWarping/Source/MotionWarping/Public/RootMotionModifier.h
@@ -150,9 +150,17 @@ struct MOTIONWARPING_API FMotionWarpingTarget
 {
 	GENERATED_BODY()
 
-	/** When the warp target is created from a component this stores the transform of the component at the time of creation, otherwise its the transform supplied by the user */
+	/** Unique name for this warp target */
 	UPROPERTY(EditAnywhere, BlueprintReadWrite, Category = "Defaults")
-	FTransform Transform;
+	FName Name;
+
+	/** When the warp target is created from a component this stores the location of the component at the time of creation, otherwise its the location supplied by the user */
+	UPROPERTY(EditAnywhere, BlueprintReadWrite, Category = "Defaults")
+	FVector Location;
+
+	/** When the warp target is created from a component this stores the rotation of the component at the time of creation, otherwise its the rotation supplied by the user */
+	UPROPERTY(EditAnywhere, BlueprintReadWrite, Category = "Defaults")
+	FRotator Rotation;
 
 	/** Optional component used to calculate the final target transform */
 	UPROPERTY(EditAnywhere, BlueprintReadWrite, Category = "Defaults")
@@ -167,12 +175,12 @@ struct MOTIONWARPING_API FMotionWarpingTarget
 	bool bFollowComponent;
 
 	FMotionWarpingTarget()
-		: Transform(FTransform::Identity), Component(nullptr), BoneName(NAME_None), bFollowComponent(false) {}
+		: Name(NAME_None), Location(FVector::ZeroVector), Rotation(FRotator::ZeroRotator), Component(nullptr), BoneName(NAME_None), bFollowComponent(false) {}
 
-	FMotionWarpingTarget(const FTransform& InTransform)
-		: Transform(InTransform), Component(nullptr), BoneName(NAME_None), bFollowComponent(false) {}
+	FMotionWarpingTarget(const FName& InName, const FTransform& InTransform)
+		: Name(InName), Location(InTransform.GetLocation()), Rotation(InTransform.Rotator()), Component(nullptr), BoneName(NAME_None), bFollowComponent(false) {}
 
-	FMotionWarpingTarget(const USceneComponent* InComp, FName InBoneName, bool bInbFollowComponent);
+	FMotionWarpingTarget(const FName& InName, const USceneComponent* InComp, FName InBoneName, bool bInbFollowComponent);
 
 	FTransform GetTargetTrasform() const;
 
@@ -182,12 +190,12 @@ struct MOTIONWARPING_API FMotionWarpingTarget
 
 	FORCEINLINE bool operator==(const FMotionWarpingTarget& Other) const
 	{
-		return Other.Transform.Equals(Transform) && Other.Component == Component && Other.BoneName == BoneName && Other.bFollowComponent == bFollowComponent;
+		return Other.Name == Name && Other.Location.Equals(Location) && Other.Rotation.Equals(Rotation) && Other.Component == Component && Other.BoneName == BoneName && Other.bFollowComponent == bFollowComponent;
 	}
 
 	FORCEINLINE bool operator!=(const FMotionWarpingTarget& Other) const
 	{
-		return !Other.Transform.Equals(Transform) || Other.Component != Component || Other.BoneName != BoneName || Other.bFollowComponent != bFollowComponent;
+		return Other.Name != Name || !Other.Location.Equals(Location) || !Other.Rotation.Equals(Rotation) || Other.Component != Component || Other.BoneName != BoneName || Other.bFollowComponent != bFollowComponent;
 	}
 
 	static FTransform GetTargetTransformFromComponent(const USceneComponent* Comp, const FName& BoneName);
