Motion Warping: Fixed bug where warp target container would not replicate after updating an existing target.

diff --git a/Engine/Plugins/Experimental/Animation/MotionWarping/Source/MotionWarping/Private/MotionWarpingComponent.cpp b/Engine/Plugins/Experimental/Animation/MotionWarping/Source/MotionWarping/Private/MotionWarpingComponent.cpp
--- a/Engine/Plugins/Experimental/Animation/MotionWarping/Source/MotionWarping/Private/MotionWarpingComponent.cpp
+++ b/Engine/Plugins/Experimental/Animation/MotionWarping/Source/MotionWarping/Private/MotionWarpingComponent.cpp
@@ -486,21 +486,30 @@ FTransform UMotionWarpingComponent::ProcessRootMotionPostConvertToWorld(const FT
 	return FinalRootMotion;
 }
 
+bool UMotionWarpingComponent::FindAndUpdateWarpTarget(const FMotionWarpingTarget& WarpTarget)
+{
+	for (int32 Idx = 0; Idx < WarpTargets.Num(); Idx++)
+	{
+		if (WarpTargets[Idx].Name == WarpTarget.Name)
+		{
+			WarpTargets[Idx] = WarpTarget;
+			return true;
+		}
+	}
+
+	return false;
+}
+
 void UMotionWarpingComponent::AddOrUpdateWarpTarget(const FMotionWarpingTarget& WarpTarget)
 {
 	if (WarpTarget.Name != NAME_None)
 	{
-		for (int32 Idx = 0; Idx < WarpTargets.Num(); Idx++)
+		// if we did not find the target, add it
+		if (!FindAndUpdateWarpTarget(WarpTarget))
 		{
-			if (WarpTargets[Idx].Name == WarpTarget.Name)
-			{
-				WarpTargets[Idx] = WarpTarget;
-				return;
-			}
+			WarpTargets.Add(WarpTarget);
 		}
 
-		WarpTargets.Add(WarpTarget);
-
 		MARK_PROPERTY_DIRTY_FROM_NAME(UMotionWarpingComponent, WarpTargets, this);
 	}
 }
diff --git a/Engine/Plugins/Experimental/Animation/MotionWarping/Source/MotionWarping/Public/MotionWarpingComponent.h b/Engine/Plugins/Experimental/Animation/MotionWarping/Source/MotionWarping/Public/MotionWarpingComponent.h
--- a/Engine/Plugins/Experimental/Animation/MotionWarping/Source/MotionWarping/Public/MotionWarpingComponent.h
+++ b/Engine/Plugins/Experimental/Animation/MotionWarping/Source/MotionWarping/Public/MotionWarpingComponent.h
@@ -190,6 +190,8 @@ protected:
 
 	void Update();
 
+	bool FindAndUpdateWarpTarget(const FMotionWarpingTarget& WarpTarget);
+
 	FTransform ProcessRootMotionPreConvertToWorld(const FTransform& InRootMotion, class UCharacterMovementComponent* CharacterMovementComponent, float DeltaSeconds);
 	
 	FTransform ProcessRootMotionPostConvertToWorld(const FTransform& InRootMotion, class UCharacterMovementComponent* CharacterMovementComponent, float DeltaSeconds);
