Simplified motion warping framework by removing bInLocalSpace option. Root motion is now always warped in mesh space.

diff --git a/Engine/Plugins/Experimental/Animation/MotionWarping/Source/MotionWarping/Private/MotionWarpingComponent.cpp b/Engine/Plugins/Experimental/Animation/MotionWarping/Source/MotionWarping/Private/MotionWarpingComponent.cpp
--- a/Engine/Plugins/Experimental/Animation/MotionWarping/Source/MotionWarping/Private/MotionWarpingComponent.cpp
+++ b/Engine/Plugins/Experimental/Animation/MotionWarping/Source/MotionWarping/Private/MotionWarpingComponent.cpp
@@ -226,8 +226,7 @@
 	UCharacterMovementComponent* CharacterMovementComp = CharacterOwner.IsValid() ? CharacterOwner->GetCharacterMovement() : nullptr;
 	if (CharacterMovementComp)
 	{
-		CharacterMovementComp->ProcessRootMotionPreConvertToWorld.BindUObject(this, &UMotionWarpingComponent::ProcessRootMotionPreConvertToWorld);
-		CharacterMovementComp->ProcessRootMotionPostConvertToWorld.BindUObject(this, &UMotionWarpingComponent::ProcessRootMotionPostConvertToWorld);
+ 		CharacterMovementComp->ProcessRootMotionPreConvertToWorld.BindUObject(this, &UMotionWarpingComponent::ProcessRootMotionPreConvertToWorld);
 	}
 }
 
@@ -423,30 +422,7 @@
 	// Apply Local Space Modifiers
 	for (URootMotionModifier* Modifier : Modifiers)
 	{
-		if (Modifier->GetState() == ERootMotionModifierState::Active && Modifier->bInLocalSpace)
-		{
-			FinalRootMotion = Modifier->ProcessRootMotion(FinalRootMotion, DeltaSeconds);
-		}
-	}
-
-	return FinalRootMotion;
-}
-
-FTransform UMotionWarpingComponent::ProcessRootMotionPostConvertToWorld(const FTransform& InRootMotion, UCharacterMovementComponent* CharacterMovementComponent, float DeltaSeconds)
-{
-#if !(UE_BUILD_SHIPPING || UE_BUILD_TEST)
-	if (FMotionWarpingCVars::CVarMotionWarpingDisable.GetValueOnGameThread() > 0)
-	{
-		return InRootMotion;
-	}
-#endif
-
-	FTransform FinalRootMotion = InRootMotion;
-
-	// Apply World Space Modifiers
-	for (URootMotionModifier* Modifier : Modifiers)
-	{
-		if (Modifier->GetState() == ERootMotionModifierState::Active && !Modifier->bInLocalSpace)
+		if (Modifier->GetState() == ERootMotionModifierState::Active)
 		{
 			FinalRootMotion = Modifier->ProcessRootMotion(FinalRootMotion, DeltaSeconds);
 		}
@@ -467,8 +443,8 @@
 				WarpedRootMotionAccum = ActorFeetLocation;
 			}
 
-			OriginalRootMotionAccum = OriginalRootMotionAccum.GetValue() + InRootMotion.GetLocation();
-			WarpedRootMotionAccum = WarpedRootMotionAccum.GetValue() + FinalRootMotion.GetLocation();
+			OriginalRootMotionAccum = OriginalRootMotionAccum.GetValue() + (CharacterOwner->GetMesh()->ConvertLocalRootMotionToWorld(FTransform(InRootMotion.GetLocation()))).GetLocation();
+			WarpedRootMotionAccum = WarpedRootMotionAccum.GetValue() + (CharacterOwner->GetMesh()->ConvertLocalRootMotionToWorld(FTransform(FinalRootMotion.GetLocation()))).GetLocation();
 
 			DrawDebugPoint(GetWorld(), OriginalRootMotionAccum.GetValue(), PointSize, FColor::Red, false, DrawDebugDuration, 0);
 			DrawDebugPoint(GetWorld(), WarpedRootMotionAccum.GetValue(), PointSize, FColor::Green, false, DrawDebugDuration, 0);
diff --git a/Engine/Plugins/Experimental/Animation/MotionWarping/Source/MotionWarping/Private/RootMotionModifier.cpp b/Engine/Plugins/Experimental/Animation/MotionWarping/Source/MotionWarping/Private/RootMotionModifier.cpp
--- a/Engine/Plugins/Experimental/Animation/MotionWarping/Source/MotionWarping/Private/RootMotionModifier.cpp
+++ b/Engine/Plugins/Experimental/Animation/MotionWarping/Source/MotionWarping/Private/RootMotionModifier.cpp
@@ -271,16 +271,13 @@
 		return FQuat::Identity;
 	}
 
-	const FTransform& CharacterTransform = CharacterOwner->GetActorTransform();
-	const FQuat CurrentRotation = CharacterTransform.GetRotation();
-	const FQuat TargetRotation = GetTargetRotation();
+	const FQuat TotalRootMotionRotation = RootMotionTotal.GetRotation();
+	const FQuat CurrentRotation = CharacterOwner->GetActorQuat() * CharacterOwner->GetBaseRotationOffset();
+	const FQuat TargetRotation = CurrentRotation.Inverse() * (GetTargetRotation() * CharacterOwner->GetBaseRotationOffset());
 	const float TimeRemaining = (EndTime - PreviousPosition) * WarpRotationTimeMultiplier;
-	const FQuat RemainingRootRotationInWorld = RootMotionTotal.GetRotation();
-	const FQuat CurrentPlusRemainingRootMotion = RemainingRootRotationInWorld * CurrentRotation;
-	const float PercentThisStep = FMath::Clamp(DeltaSeconds / TimeRemaining, 0.f, 1.f);
-	const FQuat TargetRotThisFrame = FQuat::Slerp(CurrentPlusRemainingRootMotion, TargetRotation, PercentThisStep);
-	const FQuat DeltaOut = TargetRotThisFrame * CurrentPlusRemainingRootMotion.Inverse();
-
+	const float Alpha = FMath::Clamp(DeltaSeconds / TimeRemaining, 0.f, 1.f);
+	const FQuat TargetRotThisFrame = FQuat::Slerp(TotalRootMotionRotation, TargetRotation, Alpha);
+	const FQuat DeltaOut = TargetRotThisFrame * TotalRootMotionRotation.Inverse();
 	return (DeltaOut * RootMotionDelta.GetRotation());
 }
 
@@ -292,7 +289,7 @@
 		const float CapsuleHalfHeight = CharacterOwner->GetCapsuleComponent()->GetScaledCapsuleHalfHeight();
 		const FVector CurrentLocation = (CharacterOwner->GetActorLocation() - FVector(0.f, 0.f, CapsuleHalfHeight));
 		const FVector CurrentToTarget = (GetTargetLocation() - CurrentLocation).GetSafeNormal2D();
-		const FVector FutureLocation = CurrentLocation + (bInLocalSpace ? (CharacterOwner->GetMesh()->ConvertLocalRootMotionToWorld(WarpedRootMotion)).GetTranslation() : WarpedRootMotion.GetTranslation());
+		const FVector FutureLocation = CurrentLocation + (CharacterOwner->GetMesh()->ConvertLocalRootMotionToWorld(WarpedRootMotion)).GetTranslation();
 		const FRotator CurrentRotation = CharacterOwner->GetActorRotation();
 		const FRotator FutureRotation = (WarpedRootMotion.GetRotation() * CharacterOwner->GetActorQuat()).Rotator();
 		const float Dot = FVector::DotProduct(CharacterOwner->GetActorForwardVector(), CurrentToTarget);
@@ -344,16 +341,9 @@
 		const float HorizontalOriginal = RootMotionTotal.GetTranslation().Size2D();
 		const float HorizontalTranslationWarped = !FMath::IsNearlyZero(HorizontalOriginal) ? ((HorizontalDelta * HorizontalTarget) / HorizontalOriginal) : 0.f;
 
-		if (bInLocalSpace)
-		{
-			const FTransform MeshRelativeTransform = FTransform(CharacterOwner->GetBaseRotationOffset(), CharacterOwner->GetBaseTranslationOffset());
-			const FTransform MeshTransform = MeshRelativeTransform * CharacterOwner->GetActorTransform();
-			DeltaTranslation = MeshTransform.InverseTransformPositionNoScale(GetTargetLocation()).GetSafeNormal2D() * HorizontalTranslationWarped;
-		}
-		else
-		{
-			DeltaTranslation = (GetTargetLocation() - CharacterTransform.GetLocation()).GetSafeNormal2D() * HorizontalTranslationWarped;
-		}
+		const FTransform MeshRelativeTransform = FTransform(CharacterOwner->GetBaseRotationOffset(), CharacterOwner->GetBaseTranslationOffset());
+		const FTransform MeshTransform = MeshRelativeTransform * CharacterOwner->GetActorTransform();
+		DeltaTranslation = MeshTransform.InverseTransformPositionNoScale(GetTargetLocation()).GetSafeNormal2D() * HorizontalTranslationWarped;
 
 		if (!bIgnoreZAxis)
 		{
diff --git a/Engine/Plugins/Experimental/Animation/MotionWarping/Source/MotionWarping/Private/RootMotionModifier_AdjustmentBlendWarp.cpp b/Engine/Plugins/Experimental/Animation/MotionWarping/Source/MotionWarping/Private/RootMotionModifier_AdjustmentBlendWarp.cpp
--- a/Engine/Plugins/Experimental/Animation/MotionWarping/Source/MotionWarping/Private/RootMotionModifier_AdjustmentBlendWarp.cpp
+++ b/Engine/Plugins/Experimental/Animation/MotionWarping/Source/MotionWarping/Private/RootMotionModifier_AdjustmentBlendWarp.cpp
@@ -17,7 +17,6 @@
 URootMotionModifier_AdjustmentBlendWarp::URootMotionModifier_AdjustmentBlendWarp(const FObjectInitializer& ObjectInitializer)
 	: Super(ObjectInitializer)
 {
-	bInLocalSpace = true;
 }
 
 void URootMotionModifier_AdjustmentBlendWarp::ExtractBoneTransformAtTime(FTransform& OutTransform, const FName& BoneName, float Time) const
diff --git a/Engine/Plugins/Experimental/Animation/MotionWarping/Source/MotionWarping/Private/RootMotionModifier_SkewWarp.cpp b/Engine/Plugins/Experimental/Animation/MotionWarping/Source/MotionWarping/Private/RootMotionModifier_SkewWarp.cpp
--- a/Engine/Plugins/Experimental/Animation/MotionWarping/Source/MotionWarping/Private/RootMotionModifier_SkewWarp.cpp
+++ b/Engine/Plugins/Experimental/Animation/MotionWarping/Source/MotionWarping/Private/RootMotionModifier_SkewWarp.cpp
@@ -136,7 +136,6 @@
 	if (CurrentPosition > EndTime)
 	{
 		ExtraRootMotion = UMotionWarpingUtilities::ExtractRootMotionFromAnimation(Animation.Get(), EndTime, CurrentPosition);
-		ExtraRootMotion = CharacterOwner->GetMesh()->ConvertLocalRootMotionToWorld(ExtraRootMotion);
 	}
 
 	if (bWarpTranslation && !RootMotionDelta.GetTranslation().IsNearlyZero())
@@ -144,10 +143,9 @@
 		const float CapsuleHalfHeight = CharacterOwner->GetCapsuleComponent()->GetScaledCapsuleHalfHeight();
 		const FQuat CurrentRotation = CharacterOwner->GetActorQuat();
 		const FVector CurrentLocation = (CharacterOwner->GetActorLocation() - CurrentRotation.GetUpVector() * CapsuleHalfHeight);
-		const FTransform CurrentTransform = FTransform(CurrentRotation, CurrentLocation);
 
-		const FVector DeltaTranslation = CharacterOwner->GetMesh()->ConvertLocalRootMotionToWorld(RootMotionDelta).GetLocation();
-		const FVector TotalTranslation = CharacterOwner->GetMesh()->ConvertLocalRootMotionToWorld(RootMotionTotal).GetLocation();
+		const FVector DeltaTranslation = RootMotionDelta.GetLocation();
+		const FVector TotalTranslation = RootMotionTotal.GetLocation();
 
 		FVector TargetLocation = GetTargetLocation();
 		if (bIgnoreZAxis)
@@ -155,7 +153,10 @@
 			TargetLocation.Z = CurrentLocation.Z;
 		}
 
-		const FVector WarpedTranslation = WarpTranslation(CurrentTransform, DeltaTranslation, TotalTranslation, TargetLocation) + ExtraRootMotion.GetLocation();
+		const FTransform MeshTransform = FTransform(CharacterOwner->GetBaseRotationOffset(), CharacterOwner->GetBaseTranslationOffset()) * CharacterOwner->GetActorTransform();
+		TargetLocation = MeshTransform.InverseTransformPositionNoScale(TargetLocation);
+
+		const FVector WarpedTranslation = WarpTranslation(FTransform::Identity, DeltaTranslation, TotalTranslation, TargetLocation) + ExtraRootMotion.GetLocation();
 		FinalRootMotion.SetTranslation(WarpedTranslation);
 	}
 
@@ -202,7 +203,7 @@
 		NewModifier->bWarpRotation = bInWarpRotation;
 		NewModifier->RotationType = InRotationType;
 		NewModifier->WarpRotationTimeMultiplier = InWarpRotationTimeMultiplier;
-		
+
 		InMotionWarpingComp->AddModifier(NewModifier);
 
 		return NewModifier;
diff --git a/Engine/Plugins/Experimental/Animation/MotionWarping/Source/MotionWarping/Public/MotionWarpingComponent.h b/Engine/Plugins/Experimental/Animation/MotionWarping/Source/MotionWarping/Public/MotionWarpingComponent.h
--- a/Engine/Plugins/Experimental/Animation/MotionWarping/Source/MotionWarping/Public/MotionWarpingComponent.h
+++ b/Engine/Plugins/Experimental/Animation/MotionWarping/Source/MotionWarping/Public/MotionWarpingComponent.h
@@ -193,6 +193,4 @@
 	bool FindAndUpdateWarpTarget(const FMotionWarpingTarget& WarpTarget);
 
 	FTransform ProcessRootMotionPreConvertToWorld(const FTransform& InRootMotion, class UCharacterMovementComponent* CharacterMovementComponent, float DeltaSeconds);
-	
-	FTransform ProcessRootMotionPostConvertToWorld(const FTransform& InRootMotion, class UCharacterMovementComponent* CharacterMovementComponent, float DeltaSeconds);
 };
diff --git a/Engine/Plugins/Experimental/Animation/MotionWarping/Source/MotionWarping/Public/RootMotionModifier.h b/Engine/Plugins/Experimental/Animation/MotionWarping/Source/MotionWarping/Public/RootMotionModifier.h
--- a/Engine/Plugins/Experimental/Animation/MotionWarping/Source/MotionWarping/Public/RootMotionModifier.h
+++ b/Engine/Plugins/Experimental/Animation/MotionWarping/Source/MotionWarping/Public/RootMotionModifier.h
@@ -86,10 +86,6 @@
 	UPROPERTY(BlueprintReadOnly, Category = "Defaults")
 	float Weight = 0.f;
 
-	/** Whether this modifier runs before the extracted root motion is converted to world space or after */
-	UPROPERTY(BlueprintReadOnly, Category = "Defaults")
-	bool bInLocalSpace = false;
-
 	/** Character owner transform at the time this modifier becomes active */
 	UPROPERTY(Transient, BlueprintReadOnly, Category = "Defaults")
 	FTransform StartTransform;
@@ -328,7 +324,7 @@
 	UPROPERTY(EditAnywhere, BlueprintReadWrite, Category = "Config")
 	FVector Scale = FVector(1.f);
 
-	URootMotionModifier_Scale(const FObjectInitializer& ObjectInitializer) : Super(ObjectInitializer) { bInLocalSpace = true; }
+	URootMotionModifier_Scale(const FObjectInitializer& ObjectInitializer) : Super(ObjectInitializer) {}
 
 	virtual FTransform ProcessRootMotion(const FTransform& InRootMotion, float DeltaSeconds) override
 	{
