Motion Warping: Added support for injecting translation when there is no translation to warp in the animation
#fly [REDACTED_HANDLE]

diff --git a/Engine/Plugins/Experimental/Animation/MotionWarping/Source/MotionWarping/Private/RootMotionModifier.cpp b/Engine/Plugins/Experimental/Animation/MotionWarping/Source/MotionWarping/Private/RootMotionModifier.cpp
--- a/Engine/Plugins/Experimental/Animation/MotionWarping/Source/MotionWarping/Private/RootMotionModifier.cpp
+++ b/Engine/Plugins/Experimental/Animation/MotionWarping/Source/MotionWarping/Private/RootMotionModifier.cpp
@@ -169,7 +169,9 @@
 			ActualStartTime = PreviousPosition;
 
 			const float CapsuleHalfHeight = CharacterOwner->GetCapsuleComponent()->GetScaledCapsuleHalfHeight();
-			StartTransform = FTransform(CharacterOwner->GetActorQuat(), (CharacterOwner->GetActorLocation() - FVector(0.f, 0.f, CapsuleHalfHeight)));
+			const FQuat CurrentRotation = CharacterOwner->GetActorQuat();
+			const FVector CurrentLocation = (CharacterOwner->GetActorLocation() - CurrentRotation.GetUpVector() * CapsuleHalfHeight);
+			StartTransform = FTransform(CurrentRotation, CurrentLocation);
 
 			OnActivateDelegate.ExecuteIfBound(OwnerComp, this);
 		}
@@ -244,6 +246,19 @@
 	}
 }
 
+void URootMotionModifier_Warp::OnTargetTransformChanged()
+{
+	if (const ACharacter* CharacterOwner = GetCharacterOwner())
+	{
+		ActualStartTime = PreviousPosition;
+
+		const float CapsuleHalfHeight = CharacterOwner->GetCapsuleComponent()->GetScaledCapsuleHalfHeight();
+		const FQuat CurrentRotation = CharacterOwner->GetActorQuat();
+		const FVector CurrentLocation = (CharacterOwner->GetActorLocation() - CurrentRotation.GetUpVector() * CapsuleHalfHeight);
+		StartTransform = FTransform(CurrentRotation, CurrentLocation);
+	}
+}
+
 FQuat URootMotionModifier_Warp::GetTargetRotation() const
 {
 	if (RotationType == EMotionWarpRotationType::Default)
diff --git a/Engine/Plugins/Experimental/Animation/MotionWarping/Source/MotionWarping/Private/RootMotionModifier_SkewWarp.cpp b/Engine/Plugins/Experimental/Animation/MotionWarping/Source/MotionWarping/Private/RootMotionModifier_SkewWarp.cpp
--- a/Engine/Plugins/Experimental/Animation/MotionWarping/Source/MotionWarping/Private/RootMotionModifier_SkewWarp.cpp
+++ b/Engine/Plugins/Experimental/Animation/MotionWarping/Source/MotionWarping/Private/RootMotionModifier_SkewWarp.cpp
@@ -138,7 +138,7 @@
 		ExtraRootMotion = UMotionWarpingUtilities::ExtractRootMotionFromAnimation(Animation.Get(), EndTime, CurrentPosition);
 	}
 
-	if (bWarpTranslation && !RootMotionDelta.GetTranslation().IsNearlyZero())
+	if (bWarpTranslation)
 	{
 		const float CapsuleHalfHeight = CharacterOwner->GetCapsuleComponent()->GetScaledCapsuleHalfHeight();
 		const FQuat CurrentRotation = CharacterOwner->GetActorQuat();
@@ -153,11 +153,31 @@
 			TargetLocation.Z = CurrentLocation.Z;
 		}
 
-		const FTransform MeshTransform = FTransform(CharacterOwner->GetBaseRotationOffset(), CharacterOwner->GetBaseTranslationOffset()) * CharacterOwner->GetActorTransform();
-		TargetLocation = MeshTransform.InverseTransformPositionNoScale(TargetLocation);
+		// if there is translation in the animation, warp it
+		if (!TotalTranslation.IsNearlyZero())
+		{
+			if (!DeltaTranslation.IsNearlyZero())
+			{
+				const FTransform MeshTransform = FTransform(CharacterOwner->GetBaseRotationOffset(), CharacterOwner->GetBaseTranslationOffset()) * CharacterOwner->GetActorTransform();
+				TargetLocation = MeshTransform.InverseTransformPositionNoScale(TargetLocation);
 
-		const FVector WarpedTranslation = WarpTranslation(FTransform::Identity, DeltaTranslation, TotalTranslation, TargetLocation) + ExtraRootMotion.GetLocation();
-		FinalRootMotion.SetTranslation(WarpedTranslation);
+				const FVector WarpedTranslation = WarpTranslation(FTransform::Identity, DeltaTranslation, TotalTranslation, TargetLocation) + ExtraRootMotion.GetLocation();
+				FinalRootMotion.SetTranslation(WarpedTranslation);
+			}
+		}
+		// if there is no translation in the animation, add it
+		else
+		{
+			float Alpha = FMath::Clamp((CurrentPosition - ActualStartTime) / (EndTime - ActualStartTime), 0.f, 1.f);
+			Alpha = FAlphaBlend::AlphaToBlendOption(Alpha, AddTranslationEasingFunc, AddTranslationEasingCurve);
+
+			const FVector NextLocation = FMath::Lerp<FVector, float>(StartTransform.GetLocation(), TargetLocation, Alpha);
+			FVector FinalDeltaTranslation = (NextLocation - CurrentLocation);
+			FinalDeltaTranslation = (CurrentRotation.Inverse() * (TargetLocation - CurrentLocation).ToOrientationQuat()).GetForwardVector() * FinalDeltaTranslation.Size();
+			FinalDeltaTranslation = CharacterOwner->GetBaseRotationOffset().UnrotateVector(FinalDeltaTranslation);
+
+			FinalRootMotion.SetTranslation(FinalDeltaTranslation + ExtraRootMotion.GetLocation());
+		}
 	}
 
 	if(bWarpRotation)
diff --git a/Engine/Plugins/Experimental/Animation/MotionWarping/Source/MotionWarping/Public/RootMotionModifier.h b/Engine/Plugins/Experimental/Animation/MotionWarping/Source/MotionWarping/Public/RootMotionModifier.h
--- a/Engine/Plugins/Experimental/Animation/MotionWarping/Source/MotionWarping/Public/RootMotionModifier.h
+++ b/Engine/Plugins/Experimental/Animation/MotionWarping/Source/MotionWarping/Public/RootMotionModifier.h
@@ -4,6 +4,7 @@
 
 #include "CoreMinimal.h"
 #include "Animation/AnimSequence.h"
+#include "AlphaBlend.h"
 #include "RootMotionModifier.generated.h"
 
 class UMotionWarpingComponent;
@@ -252,6 +253,14 @@
 	UPROPERTY(EditAnywhere, BlueprintReadWrite, Category = "Config", meta = (EditCondition = "bWarpTranslation"))
 	bool bIgnoreZAxis = true;
 
+	/** Easing function used when adding translation. Only relevant when there is no translation in the animation */
+	UPROPERTY(EditAnywhere, BlueprintReadWrite, Category = "Config", meta = (EditCondition = "bWarpTranslation"))
+	EAlphaBlendOption AddTranslationEasingFunc = EAlphaBlendOption::Linear;
+
+	/** Custom curve used to add translation when there is none to warp. Only relevant when AddTranslationEasingFunc is set to Custom */
+	UPROPERTY(EditAnywhere, BlueprintReadWrite, Category = "Blend", meta = (EditCondition = "AddTranslationEasingFunc==EAlphaBlendOption::Custom", EditConditionHides))
+	TObjectPtr<class UCurveFloat> AddTranslationEasingCurve = nullptr;
+
 	/** Whether to warp the rotation component of the root motion */
 	UPROPERTY(EditAnywhere, BlueprintReadWrite, Category = "Config")
 	bool bWarpRotation = true;
@@ -274,7 +283,7 @@
 	//~ End FRootMotionModifier Interface
 
 	/** Event called during update if the target transform changes while the warping is active */
-	virtual void OnTargetTransformChanged() {}
+	virtual void OnTargetTransformChanged();
 
 #if !(UE_BUILD_SHIPPING || UE_BUILD_TEST)
 	void PrintLog(const FString& Name, const FTransform& OriginalRootMotion, const FTransform& WarpedRootMotion) const;
