Motion Warping: Abort warping if the user jumps to a different position in the animation (e.g jumps to a different section in the montage) while the warping window is active.

diff --git a/Engine/Plugins/Experimental/Animation/MotionWarping/Source/MotionWarping/Private/MotionWarpingComponent.cpp b/Engine/Plugins/Experimental/Animation/MotionWarping/Source/MotionWarping/Private/MotionWarpingComponent.cpp
--- a/Engine/Plugins/Experimental/Animation/MotionWarping/Source/MotionWarping/Private/MotionWarpingComponent.cpp
+++ b/Engine/Plugins/Experimental/Animation/MotionWarping/Source/MotionWarping/Private/MotionWarpingComponent.cpp
@@ -263,12 +263,13 @@ void UMotionWarpingComponent::DisableAllRootMotionModifiers()
 	}
 }
 
-void UMotionWarpingComponent::Update()
+void UMotionWarpingComponent::Update(float DeltaSeconds)
 {
 	const ACharacter* Character = GetCharacterOwner();
 	check(Character);
 
 	FMotionWarpingUpdateContext Context;
+	Context.DeltaSeconds = DeltaSeconds;
 
 	// When replaying saved moves we need to look at the contributor to root motion back then.
 	if (Character->bClientUpdating)
@@ -283,7 +284,8 @@ void UMotionWarpingComponent::Update()
 		{
 			Context.Animation = SavedMove->RootMotionMontage.Get();
 			Context.CurrentPosition = SavedMove->RootMotionTrackPosition;
-			Context.PreviousPosition = FMath::Max(Context.CurrentPosition - SavedMove->DeltaTime * SavedMove->RootMotionMontage->RateScale, 0.f);
+			Context.PreviousPosition = SavedMove->RootMotionPreviousTrackPosition;
+			Context.PlayRate = SavedMove->RootMotionPlayRateWithScale;
 		}
 	}
 	else // If we are not replaying a move, just use the current root motion montage
@@ -297,6 +299,7 @@ void UMotionWarpingComponent::Update()
 			Context.CurrentPosition = RootMotionMontageInstance->GetPosition();
 			Context.PreviousPosition = RootMotionMontageInstance->GetPreviousPosition();
 			Context.Weight = RootMotionMontageInstance->GetWeight();
+			Context.PlayRate = RootMotionMontageInstance->Montage->RateScale * RootMotionMontageInstance->GetPlayRate();
 		}
 	}
 
@@ -415,7 +418,7 @@ FTransform UMotionWarpingComponent::ProcessRootMotionPreConvertToWorld(const FTr
 #endif
 
 	// Check for warping windows and update modifier states
-	Update();
+	Update(DeltaSeconds);
 
 	FTransform FinalRootMotion = InRootMotion;
 
diff --git a/Engine/Plugins/Experimental/Animation/MotionWarping/Source/MotionWarping/Private/RootMotionModifier.cpp b/Engine/Plugins/Experimental/Animation/MotionWarping/Source/MotionWarping/Private/RootMotionModifier.cpp
--- a/Engine/Plugins/Experimental/Animation/MotionWarping/Source/MotionWarping/Private/RootMotionModifier.cpp
+++ b/Engine/Plugins/Experimental/Animation/MotionWarping/Source/MotionWarping/Private/RootMotionModifier.cpp
@@ -126,6 +126,21 @@ void URootMotionModifier::Update(const FMotionWarpingUpdateContext& Context)
 		return;
 	}
 
+	// Mark for removal if we jumped to a position outside the warping window
+	if (State == ERootMotionModifierState::Active && PreviousPosition < EndTime && (CurrentPosition > EndTime || CurrentPosition < StartTime))
+	{
+		const float ExpectedDelta = Context.DeltaSeconds * Context.PlayRate;
+		const float ActualDelta = CurrentPosition - PreviousPosition;
+		if (!FMath::IsNearlyZero(FMath::Abs(ActualDelta - ExpectedDelta), KINDA_SMALL_NUMBER))
+		{
+			UE_LOG(LogMotionWarping, Verbose, TEXT("MotionWarping: Marking RootMotionModifier for removal. Reason: CurrentPosition manually changed. PrevPos: %f CurrPos: %f DeltaTime: %f ExpectedDelta: %f ActualDelta: %f"),
+				PreviousPosition, CurrentPosition, Context.DeltaSeconds, ExpectedDelta, ActualDelta);
+
+			SetState(ERootMotionModifierState::MarkedForRemoval);
+			return;
+		}
+	}
+
 	// Check if we are inside the warping window
 	if (PreviousPosition >= StartTime && PreviousPosition < EndTime)
 	{
diff --git a/Engine/Plugins/Experimental/Animation/MotionWarping/Source/MotionWarping/Public/MotionWarpingComponent.h b/Engine/Plugins/Experimental/Animation/MotionWarping/Source/MotionWarping/Public/MotionWarpingComponent.h
--- a/Engine/Plugins/Experimental/Animation/MotionWarping/Source/MotionWarping/Public/MotionWarpingComponent.h
+++ b/Engine/Plugins/Experimental/Animation/MotionWarping/Source/MotionWarping/Public/MotionWarpingComponent.h
@@ -188,7 +188,7 @@ protected:
 	TOptional<FVector> WarpedRootMotionAccum;
 #endif
 
-	void Update();
+	void Update(float DeltaSeconds);
 
 	bool FindAndUpdateWarpTarget(const FMotionWarpingTarget& WarpTarget);
 
diff --git a/Engine/Plugins/Experimental/Animation/MotionWarping/Source/MotionWarping/Public/RootMotionModifier.h b/Engine/Plugins/Experimental/Animation/MotionWarping/Source/MotionWarping/Public/RootMotionModifier.h
--- a/Engine/Plugins/Experimental/Animation/MotionWarping/Source/MotionWarping/Public/RootMotionModifier.h
+++ b/Engine/Plugins/Experimental/Animation/MotionWarping/Source/MotionWarping/Public/RootMotionModifier.h
@@ -32,6 +32,12 @@ struct FMotionWarpingUpdateContext
 
 	UPROPERTY()
 	float Weight = 0.f;
+
+	UPROPERTY()
+	float PlayRate = 1.f;
+
+	UPROPERTY()
+	float DeltaSeconds = 0.f;
 };
 
 /** The possible states of a Root Motion Modifier */
