Moving Motion Warping plugin out of experimental

diff --git a/Engine/Plugins/Animation/MotionWarping/MotionWarping.uplugin b/Engine/Plugins/Animation/MotionWarping/MotionWarping.uplugin
new file mode 100644
--- /dev/null
+++ b/Engine/Plugins/Animation/MotionWarping/MotionWarping.uplugin
@@ -0,0 +1,25 @@
+{
+	"FileVersion": 3,
+	"Version": 1,
+	"VersionName": "0.1",
+	"FriendlyName": "Motion Warping",
+	"Description": "",
+	"Category": "Animation",
+	"CreatedBy": "Epic Games, Inc.",
+	"CreatedByURL": "http://epicgames.com",
+	"DocsURL": "",
+	"MarketplaceURL": "",
+	"SupportURL": "",
+	"EnabledByDefault": false,
+	"CanContainContent": true,
+	"IsBetaVersion": false,
+	"IsExperimentalVersion": true,
+	"Installed": false,
+	"Modules": [
+		{
+			"Name": "MotionWarping",
+			"Type": "Runtime",
+			"LoadingPhase": "Default"
+		}
+	]
+}
\ No newline at end of file
diff --git a/Engine/Plugins/Animation/MotionWarping/Source/MotionWarping/MotionWarping.Build.cs b/Engine/Plugins/Animation/MotionWarping/Source/MotionWarping/MotionWarping.Build.cs
new file mode 100644
--- /dev/null
+++ b/Engine/Plugins/Animation/MotionWarping/Source/MotionWarping/MotionWarping.Build.cs
@@ -0,0 +1,52 @@
+// Copyright Epic Games, Inc. All Rights Reserved.
+
+using UnrealBuildTool;
+
+public class MotionWarping : ModuleRules
+{
+	public MotionWarping(ReadOnlyTargetRules Target) : base(Target)
+	{
+		PCHUsage = ModuleRules.PCHUsageMode.UseExplicitOrSharedPCHs;
+		
+		PublicIncludePaths.AddRange(
+			new string[] {
+				// ... add public include paths required here ...
+			}
+			);
+				
+		
+		PrivateIncludePaths.AddRange(
+			new string[] {
+				// ... add other private include paths required here ...
+			}
+			);
+			
+		
+		PublicDependencyModuleNames.AddRange(
+			new string[]
+			{
+				"Core",
+				"CoreUObject",		
+				"Engine",
+				"NetCore"
+			}
+			);
+			
+		
+		PrivateDependencyModuleNames.AddRange(
+			new string[]
+			{
+				"Slate",
+				"SlateCore"
+            }
+			);
+		
+		
+		DynamicallyLoadedModuleNames.AddRange(
+			new string[]
+			{
+				// ... add any modules that your module loads dynamically here ...
+			}
+			);
+	}
+}
diff --git a/Engine/Plugins/Animation/MotionWarping/Source/MotionWarping/Private/AnimNotifyState_MotionWarping.cpp b/Engine/Plugins/Animation/MotionWarping/Source/MotionWarping/Private/AnimNotifyState_MotionWarping.cpp
new file mode 100644
--- /dev/null
+++ b/Engine/Plugins/Animation/MotionWarping/Source/MotionWarping/Private/AnimNotifyState_MotionWarping.cpp
@@ -0,0 +1,116 @@
+// Copyright Epic Games, Inc. All Rights Reserved.
+
+#include "AnimNotifyState_MotionWarping.h"
+#include "GameFramework/Actor.h"
+#include "MotionWarpingComponent.h"
+#include "RootMotionModifier.h"
+
+#if WITH_EDITOR
+#include "Logging/MessageLog.h"
+#include "Misc/UObjectToken.h"
+#include "UObject/ObjectSaveContext.h"
+#endif
+
+UAnimNotifyState_MotionWarping::UAnimNotifyState_MotionWarping(const FObjectInitializer& ObjectInitializer)
+	: Super(ObjectInitializer)
+{
+}
+
+void UAnimNotifyState_MotionWarping::OnBecomeRelevant(UMotionWarpingComponent* MotionWarpingComp, const UAnimSequenceBase* Animation, float StartTime, float EndTime) const
+{
+	URootMotionModifier* RootMotionModifierNew = AddRootMotionModifier(MotionWarpingComp, Animation, StartTime, EndTime);
+
+	if (RootMotionModifierNew)
+	{
+		if (!RootMotionModifierNew->OnActivateDelegate.IsBound())
+		{
+			RootMotionModifierNew->OnActivateDelegate.BindDynamic(this, &UAnimNotifyState_MotionWarping::OnRootMotionModifierActivate);
+		}
+
+		if (!RootMotionModifierNew->OnUpdateDelegate.IsBound())
+		{
+			RootMotionModifierNew->OnUpdateDelegate.BindDynamic(this, &UAnimNotifyState_MotionWarping::OnRootMotionModifierUpdate);
+		}
+
+		if (!RootMotionModifierNew->OnDeactivateDelegate.IsBound())
+		{
+			RootMotionModifierNew->OnDeactivateDelegate.BindDynamic(this, &UAnimNotifyState_MotionWarping::OnRootMotionModifierDeactivate);
+		}
+	}
+}
+
+URootMotionModifier* UAnimNotifyState_MotionWarping::AddRootMotionModifier_Implementation(UMotionWarpingComponent* MotionWarpingComp, const UAnimSequenceBase* Animation, float StartTime, float EndTime) const
+{
+	if (MotionWarpingComp && RootMotionModifier)
+	{
+		return MotionWarpingComp->AddModifierFromTemplate(RootMotionModifier, Animation, StartTime, EndTime);
+	}
+
+	return nullptr;
+}
+
+void UAnimNotifyState_MotionWarping::OnRootMotionModifierActivate(UMotionWarpingComponent* MotionWarpingComp, URootMotionModifier* Modifier)
+{
+	// Notify blueprint
+	OnWarpBegin(MotionWarpingComp, Modifier);
+}
+
+void UAnimNotifyState_MotionWarping::OnRootMotionModifierUpdate(UMotionWarpingComponent* MotionWarpingComp, URootMotionModifier* Modifier)
+{
+	// Notify blueprint
+	OnWarpUpdate(MotionWarpingComp, Modifier);
+}
+
+void UAnimNotifyState_MotionWarping::OnRootMotionModifierDeactivate(UMotionWarpingComponent* MotionWarpingComp, URootMotionModifier* Modifier)
+{
+	// Notify blueprint
+	OnWarpEnd(MotionWarpingComp, Modifier);
+}
+
+#if WITH_EDITOR
+void UAnimNotifyState_MotionWarping::ValidateAssociatedAssets()
+{
+	static const FName NAME_AssetCheck("AssetCheck");
+
+	if(UAnimSequenceBase* ContainingAsset = Cast<UAnimSequenceBase>(GetContainingAsset()))
+	{
+		if (RootMotionModifier == nullptr)
+		{
+			FMessageLog AssetCheckLog(NAME_AssetCheck);
+
+			const FText MessageLooping = FText::Format(
+				NSLOCTEXT("AnimNotify", "MotionWarping_InvalidRootMotionModifier", "Motion Warping window in {0} doesn't have a valid RootMotionModifier"),
+				FText::AsCultureInvariant(GetNameSafe(ContainingAsset)));
+			AssetCheckLog.Warning()
+				->AddToken(FUObjectToken::Create(ContainingAsset))
+				->AddToken(FTextToken::Create(MessageLooping));
+
+			if (GIsEditor)
+			{
+				const bool bForce = true;
+				AssetCheckLog.Notify(MessageLooping, EMessageSeverity::Warning, bForce);
+			}
+		}
+		else if (const URootMotionModifier_Warp* RootMotionModifierWarp = Cast<URootMotionModifier_Warp>(RootMotionModifier))
+		{
+			if (RootMotionModifierWarp->WarpTargetName.IsNone())
+			{
+				FMessageLog AssetCheckLog(NAME_AssetCheck);
+
+				const FText MessageLooping = FText::Format(
+					NSLOCTEXT("AnimNotify", "MotionWarping_InvalidWarpTargetName", "Motion Warping window in {0} doesn't specify a valid Warp Target Name"),
+					FText::AsCultureInvariant(GetNameSafe(ContainingAsset)));
+				AssetCheckLog.Warning()
+					->AddToken(FUObjectToken::Create(ContainingAsset))
+					->AddToken(FTextToken::Create(MessageLooping));
+
+				if (GIsEditor)
+				{
+					const bool bForce = true;
+					AssetCheckLog.Notify(MessageLooping, EMessageSeverity::Warning, bForce);
+				}
+			}
+		}
+	}
+}
+#endif
\ No newline at end of file
diff --git a/Engine/Plugins/Animation/MotionWarping/Source/MotionWarping/Private/MotionWarping.cpp b/Engine/Plugins/Animation/MotionWarping/Source/MotionWarping/Private/MotionWarping.cpp
new file mode 100644
--- /dev/null
+++ b/Engine/Plugins/Animation/MotionWarping/Source/MotionWarping/Private/MotionWarping.cpp
@@ -0,0 +1,19 @@
+// Copyright Epic Games, Inc. All Rights Reserved.
+
+#include "MotionWarping.h"
+
+#define LOCTEXT_NAMESPACE "MotionWarpingModule"
+
+void FMotionWarpingModule::StartupModule()
+{
+
+}
+
+void FMotionWarpingModule::ShutdownModule()
+{
+
+}
+
+#undef LOCTEXT_NAMESPACE
+	
+IMPLEMENT_MODULE(FMotionWarpingModule, MotionWarping)
\ No newline at end of file
diff --git a/Engine/Plugins/Animation/MotionWarping/Source/MotionWarping/Private/MotionWarpingComponent.cpp b/Engine/Plugins/Animation/MotionWarping/Source/MotionWarping/Private/MotionWarpingComponent.cpp
new file mode 100644
--- /dev/null
+++ b/Engine/Plugins/Animation/MotionWarping/Source/MotionWarping/Private/MotionWarpingComponent.cpp
@@ -0,0 +1,541 @@
+// Copyright Epic Games, Inc. All Rights Reserved.
+
+#include "MotionWarpingComponent.h"
+
+#include "RootMotionModifier.h"
+#include "Animation/AnimSequenceBase.h"
+#include "Animation/AnimMontage.h"
+#include "Animation/AnimInstance.h"
+#include "Animation/AnimationPoseData.h"
+#include "DrawDebugHelpers.h"
+#include "GameFramework/Character.h"
+#include "GameFramework/CharacterMovementComponent.h"
+#include "AnimNotifyState_MotionWarping.h"
+#include "Net/UnrealNetwork.h"
+
+DEFINE_LOG_CATEGORY(LogMotionWarping);
+
+#if !(UE_BUILD_SHIPPING || UE_BUILD_TEST)
+TAutoConsoleVariable<int32> FMotionWarpingCVars::CVarMotionWarpingDisable(TEXT("a.MotionWarping.Disable"), 0, TEXT("Disable Motion Warping"), ECVF_Cheat);
+TAutoConsoleVariable<int32> FMotionWarpingCVars::CVarMotionWarpingDebug(TEXT("a.MotionWarping.Debug"), 0, TEXT("0: Disable, 1: Only Log, 2: Only DrawDebug, 3: Log and DrawDebug"), ECVF_Cheat);
+TAutoConsoleVariable<float> FMotionWarpingCVars::CVarMotionWarpingDrawDebugDuration(TEXT("a.MotionWarping.DrawDebugLifeTime"), 1.f, TEXT("Time in seconds each draw debug persists.\nRequires 'a.MotionWarping.Debug 2'"), ECVF_Cheat);
+#endif
+
+// UMotionWarpingUtilities
+///////////////////////////////////////////////////////////////////////
+
+void UMotionWarpingUtilities::ExtractLocalSpacePose(const UAnimSequenceBase* Animation, const FBoneContainer& BoneContainer, float Time, bool bExtractRootMotion, FCompactPose& OutPose)
+{
+	OutPose.SetBoneContainer(&BoneContainer);
+
+	FBlendedCurve Curve;
+	Curve.InitFrom(BoneContainer);
+
+	FAnimExtractContext Context(Time, bExtractRootMotion);
+
+	UE::Anim::FStackAttributeContainer Attributes;
+	FAnimationPoseData AnimationPoseData(OutPose, Curve, Attributes);
+	if (const UAnimSequence* AnimSequence = Cast<UAnimSequence>(Animation))
+	{
+		AnimSequence->GetBonePose(AnimationPoseData, Context);
+	}
+	else if (const UAnimMontage* AnimMontage = Cast<UAnimMontage>(Animation))
+	{
+		const FAnimTrack& AnimTrack = AnimMontage->SlotAnimTracks[0].AnimTrack;
+		AnimTrack.GetAnimationPose(AnimationPoseData, Context);
+	}
+}
+
+void UMotionWarpingUtilities::ExtractComponentSpacePose(const UAnimSequenceBase* Animation, const FBoneContainer& BoneContainer, float Time, bool bExtractRootMotion, FCSPose<FCompactPose>& OutPose)
+{
+	FCompactPose Pose;
+	ExtractLocalSpacePose(Animation, BoneContainer, Time, bExtractRootMotion, Pose);
+	OutPose.InitPose(MoveTemp(Pose));
+}
+
+FTransform UMotionWarpingUtilities::ExtractRootMotionFromAnimation(const UAnimSequenceBase* Animation, float StartTime, float EndTime)
+{
+	if (const UAnimMontage* Anim = Cast<UAnimMontage>(Animation))
+	{
+		// This is identical to UAnimMontage::ExtractRootMotionFromTrackRange and UAnimCompositeBase::ExtractRootMotionFromTrack but ignoring bEnableRootMotion
+		// so we can extract root motion from the montage even if that flag is set to false in the AnimSequence(s)
+
+		FRootMotionMovementParams AccumulatedRootMotionParams;
+
+		if (Anim->SlotAnimTracks.Num() > 0)
+		{
+			const FAnimTrack& RootMotionAnimTrack = Anim->SlotAnimTracks[0].AnimTrack;
+
+			TArray<FRootMotionExtractionStep> RootMotionExtractionSteps;
+			RootMotionAnimTrack.GetRootMotionExtractionStepsForTrackRange(RootMotionExtractionSteps, StartTime, EndTime);
+
+			for (const FRootMotionExtractionStep& CurStep : RootMotionExtractionSteps)
+			{
+				if (CurStep.AnimSequence)
+				{
+					AccumulatedRootMotionParams.Accumulate(CurStep.AnimSequence->ExtractRootMotionFromRange(CurStep.StartPosition, CurStep.EndPosition));
+				}
+			}
+		}
+
+		return AccumulatedRootMotionParams.GetRootMotionTransform();
+	}
+
+	if (const UAnimSequence* Anim = Cast<UAnimSequence>(Animation))
+	{
+		return Anim->ExtractRootMotionFromRange(StartTime, EndTime);
+	}
+
+	return FTransform::Identity;
+}
+
+FTransform UMotionWarpingUtilities::ExtractRootTransformFromAnimation(const UAnimSequenceBase* Animation, float Time)
+{
+	if (const UAnimMontage* AnimMontage = Cast<UAnimMontage>(Animation))
+	{
+		if(const FAnimSegment* Segment = AnimMontage->SlotAnimTracks[0].AnimTrack.GetSegmentAtTime(Time))
+		{
+			if (const UAnimSequence* AnimSequence = Cast<UAnimSequence>(Segment->GetAnimReference()))
+			{
+				const float AnimSequenceTime = Segment->ConvertTrackPosToAnimPos(Time);
+				return AnimSequence->ExtractRootTrackTransform(AnimSequenceTime, nullptr);
+			}	
+		}
+	}
+	else if (const UAnimSequence* AnimSequence = Cast<UAnimSequence>(Animation))
+	{
+		return AnimSequence->ExtractRootTrackTransform(Time, nullptr);
+	}
+
+	return FTransform::Identity;
+}
+
+void UMotionWarpingUtilities::GetMotionWarpingWindowsFromAnimation(const UAnimSequenceBase* Animation, TArray<FMotionWarpingWindowData>& OutWindows)
+{
+	if(Animation)
+	{
+		OutWindows.Reset();
+
+		for (int32 Idx = 0; Idx < Animation->Notifies.Num(); Idx++)
+		{
+			const FAnimNotifyEvent& NotifyEvent = Animation->Notifies[Idx];
+			if (UAnimNotifyState_MotionWarping* Notify = Cast<UAnimNotifyState_MotionWarping>(NotifyEvent.NotifyStateClass))
+			{
+				FMotionWarpingWindowData Data;
+				Data.AnimNotify = Notify;
+				Data.StartTime = NotifyEvent.GetTriggerTime();
+				Data.EndTime = NotifyEvent.GetEndTriggerTime();
+				OutWindows.Add(Data);
+			}
+		}
+	}
+}
+
+void UMotionWarpingUtilities::GetMotionWarpingWindowsForWarpTargetFromAnimation(const UAnimSequenceBase* Animation, FName WarpTargetName, TArray<FMotionWarpingWindowData>& OutWindows)
+{
+	if (Animation && WarpTargetName != NAME_None)
+	{
+		OutWindows.Reset();
+
+		for (int32 Idx = 0; Idx < Animation->Notifies.Num(); Idx++)
+		{
+			const FAnimNotifyEvent& NotifyEvent = Animation->Notifies[Idx];
+			if (UAnimNotifyState_MotionWarping* Notify = Cast<UAnimNotifyState_MotionWarping>(NotifyEvent.NotifyStateClass))
+			{
+				if (const URootMotionModifier_Warp* Modifier = Cast<const URootMotionModifier_Warp>(Notify->RootMotionModifier))
+				{
+					if(Modifier->WarpTargetName == WarpTargetName)
+					{
+						FMotionWarpingWindowData Data;
+						Data.AnimNotify = Notify;
+						Data.StartTime = NotifyEvent.GetTriggerTime();
+						Data.EndTime = NotifyEvent.GetEndTriggerTime();
+						OutWindows.Add(Data);
+					}
+				}
+			}
+		}
+	}
+}
+
+FTransform UMotionWarpingUtilities::CalculateRootTransformRelativeToWarpPointAtTime(const ACharacter& Character, const UAnimSequenceBase* Animation, float Time, const FName& WarpPointBoneName)
+{
+	if (const USkeletalMeshComponent* Mesh = Character.GetMesh())
+	{
+		if (const UAnimInstance* AnimInstance = Mesh->GetAnimInstance())
+		{
+			const FBoneContainer& FullBoneContainer = AnimInstance->GetRequiredBones();
+			const int32 BoneIndex = FullBoneContainer.GetPoseBoneIndexForBoneName(WarpPointBoneName);
+			if (ensure(BoneIndex != INDEX_NONE))
+			{
+				TArray<FBoneIndexType> RequiredBoneIndexArray = { 0, (FBoneIndexType)BoneIndex };
+				FullBoneContainer.GetReferenceSkeleton().EnsureParentsExistAndSort(RequiredBoneIndexArray);
+
+				FBoneContainer LimitedBoneContainer(RequiredBoneIndexArray, FCurveEvaluationOption(false), *FullBoneContainer.GetAsset());
+
+				FCSPose<FCompactPose> Pose;
+				UMotionWarpingUtilities::ExtractComponentSpacePose(Animation, LimitedBoneContainer, Time, false, Pose);
+
+				// Inverse of mesh's relative rotation. Used to convert root and warp point in the animation from Y forward to X forward
+				const FTransform MeshCompRelativeRotInverse = FTransform(Character.GetBaseRotationOffset().Inverse());
+
+				const FTransform RootTransform = MeshCompRelativeRotInverse * Pose.GetComponentSpaceTransform(FCompactPoseBoneIndex(0));
+				const FTransform WarpPointTransform = MeshCompRelativeRotInverse * Pose.GetComponentSpaceTransform(FCompactPoseBoneIndex(1));
+				return RootTransform.GetRelativeTransform(WarpPointTransform);
+			}
+		}
+	}
+
+	return FTransform::Identity;
+}
+
+FTransform UMotionWarpingUtilities::CalculateRootTransformRelativeToWarpPointAtTime(const ACharacter& Character, const UAnimSequenceBase* Animation, float Time, const FTransform& WarpPointTransform)
+{
+	// Inverse of mesh's relative rotation. Used to convert root and warp point in the animation from Y forward to X forward
+	const FTransform MeshCompRelativeRotInverse = FTransform(Character.GetBaseRotationOffset().Inverse());
+	const FTransform RootTransform = MeshCompRelativeRotInverse * UMotionWarpingUtilities::ExtractRootTransformFromAnimation(Animation, Time);
+	return RootTransform.GetRelativeTransform((MeshCompRelativeRotInverse * WarpPointTransform));
+}
+
+// UMotionWarpingComponent
+///////////////////////////////////////////////////////////////////////
+
+UMotionWarpingComponent::UMotionWarpingComponent(const FObjectInitializer& ObjectInitializer)
+	: Super(ObjectInitializer)
+{
+	bWantsInitializeComponent = true;
+	SetIsReplicatedByDefault(true);
+}
+
+void UMotionWarpingComponent::GetLifetimeReplicatedProps(TArray< FLifetimeProperty >& OutLifetimeProps) const
+{
+	Super::GetLifetimeReplicatedProps(OutLifetimeProps);
+
+	FDoRepLifetimeParams Params;
+	Params.bIsPushBased = true;
+	Params.Condition = COND_SimulatedOnly;
+	DOREPLIFETIME_WITH_PARAMS_FAST(UMotionWarpingComponent, WarpTargets, Params);
+}
+
+void UMotionWarpingComponent::InitializeComponent()
+{
+	Super::InitializeComponent();
+
+	CharacterOwner = Cast<ACharacter>(GetOwner());
+
+	UCharacterMovementComponent* CharacterMovementComp = CharacterOwner.IsValid() ? CharacterOwner->GetCharacterMovement() : nullptr;
+	if (CharacterMovementComp)
+	{
+ 		CharacterMovementComp->ProcessRootMotionPreConvertToWorld.BindUObject(this, &UMotionWarpingComponent::ProcessRootMotionPreConvertToWorld);
+	}
+}
+
+bool UMotionWarpingComponent::ContainsModifier(const UAnimSequenceBase* Animation, float StartTime, float EndTime) const
+{
+	return Modifiers.ContainsByPredicate([=](const URootMotionModifier* Modifier)
+		{
+			return (Modifier->Animation == Animation && Modifier->StartTime == StartTime && Modifier->EndTime == EndTime);
+		});
+}
+
+int32 UMotionWarpingComponent::AddModifier(URootMotionModifier* Modifier)
+{
+	if (ensureAlways(Modifier))
+	{
+		UE_LOG(LogMotionWarping, Verbose, TEXT("MotionWarping: RootMotionModifier added. NetMode: %d WorldTime: %f Char: %s Animation: %s [%f %f] [%f %f] Loc: %s Rot: %s"),
+			GetWorld()->GetNetMode(), GetWorld()->GetTimeSeconds(), *GetNameSafe(GetCharacterOwner()), *GetNameSafe(Modifier->Animation.Get()), Modifier->StartTime, Modifier->EndTime, Modifier->PreviousPosition, Modifier->CurrentPosition,
+			*GetCharacterOwner()->GetActorLocation().ToString(), *GetCharacterOwner()->GetActorRotation().ToCompactString());
+
+		return Modifiers.Add(Modifier);
+	}
+
+	return INDEX_NONE;
+}
+
+void UMotionWarpingComponent::DisableAllRootMotionModifiers()
+{
+	if (Modifiers.Num() > 0)
+	{
+		for (URootMotionModifier* Modifier : Modifiers)
+		{
+			Modifier->SetState(ERootMotionModifierState::Disabled);
+		}
+	}
+}
+
+void UMotionWarpingComponent::Update(float DeltaSeconds)
+{
+	const ACharacter* Character = GetCharacterOwner();
+	check(Character);
+
+	FMotionWarpingUpdateContext Context;
+	Context.DeltaSeconds = DeltaSeconds;
+
+	// When replaying saved moves we need to look at the contributor to root motion back then.
+	if (Character->bClientUpdating)
+	{
+		const UCharacterMovementComponent* MoveComp = Character->GetCharacterMovement();
+		check(MoveComp);
+
+		const FSavedMove_Character* SavedMove = MoveComp->GetCurrentReplayedSavedMove();
+		check(SavedMove);
+
+		if(SavedMove->RootMotionMontage.IsValid())
+		{
+			Context.Animation = SavedMove->RootMotionMontage.Get();
+			Context.CurrentPosition = SavedMove->RootMotionTrackPosition;
+			Context.PreviousPosition = SavedMove->RootMotionPreviousTrackPosition;
+			Context.PlayRate = SavedMove->RootMotionPlayRateWithScale;
+		}
+	}
+	else // If we are not replaying a move, just use the current root motion montage
+	{
+		if(const FAnimMontageInstance* RootMotionMontageInstance = Character->GetRootMotionAnimMontageInstance())
+		{
+			const UAnimMontage* Montage = RootMotionMontageInstance->Montage;
+			check(Montage);
+
+			Context.Animation = Montage;
+			Context.CurrentPosition = RootMotionMontageInstance->GetPosition();
+			Context.PreviousPosition = RootMotionMontageInstance->GetPreviousPosition();
+			Context.Weight = RootMotionMontageInstance->GetWeight();
+			Context.PlayRate = RootMotionMontageInstance->Montage->RateScale * RootMotionMontageInstance->GetPlayRate();
+		}
+	}
+
+	if (Context.Animation.IsValid())
+	{
+		const UAnimSequenceBase* Animation = Context.Animation.Get();
+		const float PreviousPosition = Context.PreviousPosition;
+		const float CurrentPosition = Context.CurrentPosition;
+
+		// Loop over notifies directly in the montage, looking for Motion Warping windows
+		for (const FAnimNotifyEvent& NotifyEvent : Animation->Notifies)
+		{
+			const UAnimNotifyState_MotionWarping* MotionWarpingNotify = NotifyEvent.NotifyStateClass ? Cast<UAnimNotifyState_MotionWarping>(NotifyEvent.NotifyStateClass) : nullptr;
+			if (MotionWarpingNotify)
+			{
+				if(MotionWarpingNotify->RootMotionModifier == nullptr)
+				{
+					UE_LOG(LogMotionWarping, Warning, TEXT("MotionWarpingComponent::Update. A motion warping window in %s doesn't have a valid root motion modifier!"), *GetNameSafe(Animation));
+					continue;
+				}
+
+				const float StartTime = FMath::Clamp(NotifyEvent.GetTriggerTime(), 0.f, Animation->GetPlayLength());
+				const float EndTime = FMath::Clamp(NotifyEvent.GetEndTriggerTime(), 0.f, Animation->GetPlayLength());
+
+				if (PreviousPosition >= StartTime && PreviousPosition < EndTime)
+				{
+					if (!ContainsModifier(Animation, StartTime, EndTime))
+					{
+						MotionWarpingNotify->OnBecomeRelevant(this, Animation, StartTime, EndTime);
+					}
+				}
+			}
+		}
+
+		if(bSearchForWindowsInAnimsWithinMontages)
+		{
+			if(const UAnimMontage* Montage = Cast<const UAnimMontage>(Context.Animation.Get()))
+			{
+				// Same as before but scanning all animation within the montage
+				for (int32 SlotIdx = 0; SlotIdx < Montage->SlotAnimTracks.Num(); SlotIdx++)
+				{
+					const FAnimTrack& AnimTrack = Montage->SlotAnimTracks[SlotIdx].AnimTrack;
+
+					if (const FAnimSegment* AnimSegment = AnimTrack.GetSegmentAtTime(PreviousPosition))
+					{
+						if (const UAnimSequenceBase* AnimReference = AnimSegment->GetAnimReference())
+						{
+							for (const FAnimNotifyEvent& NotifyEvent : AnimReference->Notifies)
+							{
+								const UAnimNotifyState_MotionWarping* MotionWarpingNotify = NotifyEvent.NotifyStateClass ? Cast<UAnimNotifyState_MotionWarping>(NotifyEvent.NotifyStateClass) : nullptr;
+								if (MotionWarpingNotify)
+								{
+									if (MotionWarpingNotify->RootMotionModifier == nullptr)
+									{
+										UE_LOG(LogMotionWarping, Warning, TEXT("MotionWarpingComponent::Update. A motion warping window in %s doesn't have a valid root motion modifier!"), *GetNameSafe(AnimReference));
+										continue;
+									}
+
+									const float NotifyStartTime = FMath::Clamp(NotifyEvent.GetTriggerTime(), 0.f, AnimReference->GetPlayLength());
+									const float NotifyEndTime = FMath::Clamp(NotifyEvent.GetEndTriggerTime(), 0.f, AnimReference->GetPlayLength());
+
+									// Convert notify times from AnimSequence times to montage times
+									const float StartTime = (NotifyStartTime - AnimSegment->AnimStartTime) + AnimSegment->StartPos;
+									const float EndTime = (NotifyEndTime - AnimSegment->AnimStartTime) + AnimSegment->StartPos;
+
+									if (PreviousPosition >= StartTime && PreviousPosition < EndTime)
+									{
+										if (!ContainsModifier(Montage, StartTime, EndTime))
+										{
+											MotionWarpingNotify->OnBecomeRelevant(this, Montage, StartTime, EndTime);
+										}
+									}
+								}
+							}
+						}
+					}
+				}
+			}
+		}
+	}
+
+	OnPreUpdate.Broadcast(this);
+
+	// Update the state of all the modifiers
+	if (Modifiers.Num() > 0)
+	{
+		for (URootMotionModifier* Modifier : Modifiers)
+		{
+			Modifier->Update(Context);
+		}
+
+		// Remove the modifiers that has been marked for removal
+		Modifiers.RemoveAll([this](const URootMotionModifier* Modifier)
+		{
+			if (Modifier->GetState() == ERootMotionModifierState::MarkedForRemoval)
+			{
+				UE_LOG(LogMotionWarping, Verbose, TEXT("MotionWarping: RootMotionModifier removed. NetMode: %d WorldTime: %f Char: %s Animation: %s [%f %f] [%f %f] Loc: %s Rot: %s"),
+					GetWorld()->GetNetMode(), GetWorld()->GetTimeSeconds(), *GetNameSafe(GetCharacterOwner()), *GetNameSafe(Modifier->Animation.Get()), Modifier->StartTime, Modifier->EndTime, Modifier->PreviousPosition, Modifier->CurrentPosition,
+					*GetCharacterOwner()->GetActorLocation().ToString(), *GetCharacterOwner()->GetActorRotation().ToCompactString());
+
+				return true;
+			}
+
+			return false;
+		});
+	}
+}
+
+FTransform UMotionWarpingComponent::ProcessRootMotionPreConvertToWorld(const FTransform& InRootMotion, UCharacterMovementComponent* CharacterMovementComponent, float DeltaSeconds)
+{
+#if !(UE_BUILD_SHIPPING || UE_BUILD_TEST)
+	if (FMotionWarpingCVars::CVarMotionWarpingDisable.GetValueOnGameThread() > 0)
+	{
+		return InRootMotion;
+	}
+#endif
+
+	// Check for warping windows and update modifier states
+	Update(DeltaSeconds);
+
+	FTransform FinalRootMotion = InRootMotion;
+
+	// Apply Local Space Modifiers
+	for (URootMotionModifier* Modifier : Modifiers)
+	{
+		if (Modifier->GetState() == ERootMotionModifierState::Active)
+		{
+			FinalRootMotion = Modifier->ProcessRootMotion(FinalRootMotion, DeltaSeconds);
+		}
+	}
+
+#if !(UE_BUILD_SHIPPING || UE_BUILD_TEST)
+	const int32 DebugLevel = FMotionWarpingCVars::CVarMotionWarpingDebug.GetValueOnGameThread();
+	if (DebugLevel >= 2)
+	{
+		const float DrawDebugDuration = FMotionWarpingCVars::CVarMotionWarpingDrawDebugDuration.GetValueOnGameThread();
+		const float PointSize = 7.f;
+		const FVector ActorFeetLocation = CharacterMovementComponent->GetActorFeetLocation();
+		if (Modifiers.Num() > 0)
+		{
+			if (!OriginalRootMotionAccum.IsSet())
+			{
+				OriginalRootMotionAccum = ActorFeetLocation;
+				WarpedRootMotionAccum = ActorFeetLocation;
+			}
+
+			OriginalRootMotionAccum = OriginalRootMotionAccum.GetValue() + (CharacterOwner->GetMesh()->ConvertLocalRootMotionToWorld(FTransform(InRootMotion.GetLocation()))).GetLocation();
+			WarpedRootMotionAccum = WarpedRootMotionAccum.GetValue() + (CharacterOwner->GetMesh()->ConvertLocalRootMotionToWorld(FTransform(FinalRootMotion.GetLocation()))).GetLocation();
+
+			DrawDebugPoint(GetWorld(), OriginalRootMotionAccum.GetValue(), PointSize, FColor::Red, false, DrawDebugDuration, 0);
+			DrawDebugPoint(GetWorld(), WarpedRootMotionAccum.GetValue(), PointSize, FColor::Green, false, DrawDebugDuration, 0);
+		}
+		else
+		{
+			OriginalRootMotionAccum.Reset();
+			WarpedRootMotionAccum.Reset();
+		}
+
+		DrawDebugPoint(GetWorld(), ActorFeetLocation, PointSize, FColor::Blue, false, DrawDebugDuration, 0);
+	}
+#endif
+
+	return FinalRootMotion;
+}
+
+bool UMotionWarpingComponent::FindAndUpdateWarpTarget(const FMotionWarpingTarget& WarpTarget)
+{
+	for (int32 Idx = 0; Idx < WarpTargets.Num(); Idx++)
+	{
+		if (WarpTargets[Idx].Name == WarpTarget.Name)
+		{
+			WarpTargets[Idx] = WarpTarget;
+			return true;
+		}
+	}
+
+	return false;
+}
+
+void UMotionWarpingComponent::AddOrUpdateWarpTarget(const FMotionWarpingTarget& WarpTarget)
+{
+	if (WarpTarget.Name != NAME_None)
+	{
+		// if we did not find the target, add it
+		if (!FindAndUpdateWarpTarget(WarpTarget))
+		{
+			WarpTargets.Add(WarpTarget);
+		}
+
+		MARK_PROPERTY_DIRTY_FROM_NAME(UMotionWarpingComponent, WarpTargets, this);
+	}
+}
+
+int32 UMotionWarpingComponent::RemoveWarpTarget(FName WarpTargetName)
+{
+	const int32 NumRemoved = WarpTargets.RemoveAll([&WarpTargetName](const FMotionWarpingTarget& WarpTarget) { return WarpTarget.Name == WarpTargetName; });
+	
+	if(NumRemoved > 0)
+	{
+		MARK_PROPERTY_DIRTY_FROM_NAME(UMotionWarpingComponent, WarpTargets, this);
+	}
+
+	return NumRemoved;
+}
+
+void UMotionWarpingComponent::AddOrUpdateWarpTargetFromTransform(FName WarpTargetName, FTransform TargetTransform)
+{
+	AddOrUpdateWarpTarget(FMotionWarpingTarget(WarpTargetName, TargetTransform));
+}
+
+void UMotionWarpingComponent::AddOrUpdateWarpTargetFromComponent(FName WarpTargetName, const USceneComponent* Component, FName BoneName, bool bFollowComponent)
+{
+	if (Component == nullptr)
+	{
+		UE_LOG(LogMotionWarping, Warning, TEXT("AddOrUpdateWarpTargetFromComponent has failed!. Reason: Invalid Component"));
+		return;
+	}
+
+	AddOrUpdateWarpTarget(FMotionWarpingTarget(WarpTargetName, Component, BoneName, bFollowComponent));
+}
+
+URootMotionModifier* UMotionWarpingComponent::AddModifierFromTemplate(URootMotionModifier* Template, const UAnimSequenceBase* Animation, float StartTime, float EndTime)
+{
+	if (ensureAlways(Template))
+	{
+		FObjectDuplicationParameters Params(Template, this);
+		URootMotionModifier* NewRootMotionModifier = CastChecked<URootMotionModifier>(StaticDuplicateObjectEx(Params));
+		
+		NewRootMotionModifier->Animation = Animation;
+		NewRootMotionModifier->StartTime = StartTime;
+		NewRootMotionModifier->EndTime = EndTime;
+
+		AddModifier(NewRootMotionModifier);
+
+		return NewRootMotionModifier;
+	}
+
+	return nullptr;
+}
\ No newline at end of file
diff --git a/Engine/Plugins/Animation/MotionWarping/Source/MotionWarping/Private/RootMotionModifier.cpp b/Engine/Plugins/Animation/MotionWarping/Source/MotionWarping/Private/RootMotionModifier.cpp
new file mode 100644
--- /dev/null
+++ b/Engine/Plugins/Animation/MotionWarping/Source/MotionWarping/Private/RootMotionModifier.cpp
@@ -0,0 +1,440 @@
+// Copyright Epic Games, Inc. All Rights Reserved.
+
+#include "RootMotionModifier.h"
+#include "GameFramework/Character.h"
+#include "GameFramework/CharacterMovementComponent.h"
+#include "Components/CapsuleComponent.h"
+#include "Components/SkeletalMeshComponent.h"
+#include "Animation/AnimMontage.h"
+#include "MotionWarpingComponent.h"
+#include "DrawDebugHelpers.h"
+#include "Animation/AnimInstance.h"
+#include "AnimNotifyState_MotionWarping.h"
+
+// FMotionWarpingTarget
+///////////////////////////////////////////////////////////////
+
+FMotionWarpingTarget::FMotionWarpingTarget(const FName& InName, const USceneComponent* InComp, FName InBoneName, bool bInFollowComponent)
+{
+	if (ensure(InComp))
+	{
+		Name = InName;
+		Component = InComp;
+		BoneName = InBoneName;
+		bFollowComponent = bInFollowComponent;
+
+		FTransform Transform = FTransform::Identity;
+		if (BoneName != NAME_None)
+		{
+			Transform = FMotionWarpingTarget::GetTargetTransformFromComponent(InComp, InBoneName);
+		}
+		else
+		{
+			Transform = InComp->GetComponentTransform();
+		}
+
+		Location = Transform.GetLocation();
+		Rotation = Transform.Rotator();
+	}
+}
+
+FTransform FMotionWarpingTarget::GetTargetTransformFromComponent(const USceneComponent* Comp, const FName& BoneName)
+{
+	if (Comp == nullptr)
+	{
+		UE_LOG(LogMotionWarping, Warning, TEXT("FMotionWarpingTarget::GetTargetTransformFromComponent: Invalid Component"));
+		return FTransform::Identity;
+	}
+
+	if (Comp->DoesSocketExist(BoneName) == false)
+	{
+		UE_LOG(LogMotionWarping, Warning, TEXT("FMotionWarpingTarget::GetTargetTransformFromComponent: Invalid Bone or Socket. Comp: %s Owner: %s BoneName: %s"),
+			*GetNameSafe(Comp), *GetNameSafe(Comp->GetOwner()), *BoneName.ToString());
+
+		return Comp->GetComponentTransform();
+	}
+
+	return Comp->GetSocketTransform(BoneName);
+}
+
+FTransform FMotionWarpingTarget::GetTargetTrasform() const
+{
+	if (Component.IsValid() && bFollowComponent)
+	{
+		if (BoneName != NAME_None)
+		{
+			return FMotionWarpingTarget::GetTargetTransformFromComponent(Component.Get(), BoneName);
+		}
+		else
+		{
+			return Component->GetComponentTransform();
+		}
+	}
+
+	return FTransform(Rotation, Location);
+}
+
+// URootMotionModifier
+///////////////////////////////////////////////////////////////
+
+URootMotionModifier::URootMotionModifier(const FObjectInitializer& ObjectInitializer)
+	: Super(ObjectInitializer)
+{
+}
+
+UMotionWarpingComponent* URootMotionModifier::GetOwnerComponent() const
+{
+	return Cast<UMotionWarpingComponent>(GetOuter());
+}
+
+ACharacter* URootMotionModifier::GetCharacterOwner() const
+{
+	UMotionWarpingComponent* OwnerComp = GetOwnerComponent();
+	return OwnerComp ? OwnerComp->GetCharacterOwner() : nullptr;
+}
+
+void URootMotionModifier::Update(const FMotionWarpingUpdateContext& Context)
+{
+	const ACharacter* CharacterOwner = GetCharacterOwner();
+	if (CharacterOwner == nullptr)
+	{
+		return;
+	}
+
+	// Mark for removal if our animation is not relevant anymore
+	if (!Context.Animation.IsValid() || Context.Animation.Get() != Animation)
+	{
+		UE_LOG(LogMotionWarping, Verbose, TEXT("MotionWarping: Marking RootMotionModifier for removal. Reason: Animation is not valid. Char: %s Current Animation: %s. Window: Animation: %s [%f %f] [%f %f]"),
+			*GetNameSafe(CharacterOwner), *GetNameSafe(Context.Animation.Get()), *GetNameSafe(Animation.Get()), StartTime, EndTime, PreviousPosition, CurrentPosition);
+
+		SetState(ERootMotionModifierState::MarkedForRemoval);
+		return;
+	}
+
+	// Update playback times and weight
+	PreviousPosition = Context.PreviousPosition;
+	CurrentPosition = Context.CurrentPosition;
+	Weight = Context.Weight;
+
+	// Mark for removal if the animation already passed the warping window
+	if (PreviousPosition >= EndTime)
+	{
+		UE_LOG(LogMotionWarping, Verbose, TEXT("MotionWarping: Marking RootMotionModifier for removal. Reason: Window has ended. Char: %s Animation: %s [%f %f] [%f %f]"),
+			*GetNameSafe(CharacterOwner), *GetNameSafe(Animation.Get()), StartTime, EndTime, PreviousPosition, CurrentPosition);
+
+		SetState(ERootMotionModifierState::MarkedForRemoval);
+		return;
+	}
+
+	// Mark for removal if we jumped to a position outside the warping window
+	if (State == ERootMotionModifierState::Active && PreviousPosition < EndTime && (CurrentPosition > EndTime || CurrentPosition < StartTime))
+	{
+		const float ExpectedDelta = Context.DeltaSeconds * Context.PlayRate;
+		const float ActualDelta = CurrentPosition - PreviousPosition;
+		if (!FMath::IsNearlyZero(FMath::Abs(ActualDelta - ExpectedDelta), KINDA_SMALL_NUMBER))
+		{
+			UE_LOG(LogMotionWarping, Verbose, TEXT("MotionWarping: Marking RootMotionModifier for removal. Reason: CurrentPosition manually changed. PrevPos: %f CurrPos: %f DeltaTime: %f ExpectedDelta: %f ActualDelta: %f"),
+				PreviousPosition, CurrentPosition, Context.DeltaSeconds, ExpectedDelta, ActualDelta);
+
+			SetState(ERootMotionModifierState::MarkedForRemoval);
+			return;
+		}
+	}
+
+	// Check if we are inside the warping window
+	if (PreviousPosition >= StartTime && PreviousPosition < EndTime)
+	{
+		// If we were waiting, switch to active
+		if (GetState() == ERootMotionModifierState::Waiting)
+		{
+			SetState(ERootMotionModifierState::Active);
+		}
+	}
+
+	if (State == ERootMotionModifierState::Active)
+	{
+		if (UMotionWarpingComponent* OwnerComp = GetOwnerComponent())
+		{
+			OnUpdateDelegate.ExecuteIfBound(OwnerComp, this);
+		}
+	}
+}
+
+void URootMotionModifier::SetState(ERootMotionModifierState NewState)
+{
+	if (State != NewState)
+	{
+		ERootMotionModifierState LastState = State;
+
+		State = NewState;
+
+		OnStateChanged(LastState);
+	}
+}
+
+void URootMotionModifier::OnStateChanged(ERootMotionModifierState LastState)
+{
+	if (UMotionWarpingComponent* OwnerComp = GetOwnerComponent())
+	{
+		if (LastState != ERootMotionModifierState::Active && State == ERootMotionModifierState::Active)
+		{
+			const ACharacter* CharacterOwner = OwnerComp->GetCharacterOwner();
+			check(CharacterOwner);
+			
+			ActualStartTime = PreviousPosition;
+
+			const float CapsuleHalfHeight = CharacterOwner->GetCapsuleComponent()->GetScaledCapsuleHalfHeight();
+			const FQuat CurrentRotation = CharacterOwner->GetActorQuat();
+			const FVector CurrentLocation = (CharacterOwner->GetActorLocation() - CurrentRotation.GetUpVector() * CapsuleHalfHeight);
+			StartTransform = FTransform(CurrentRotation, CurrentLocation);
+
+			OnActivateDelegate.ExecuteIfBound(OwnerComp, this);
+		}
+		else if (LastState == ERootMotionModifierState::Active && (State == ERootMotionModifierState::Disabled || State == ERootMotionModifierState::MarkedForRemoval))
+		{
+			OnDeactivateDelegate.ExecuteIfBound(OwnerComp, this);
+		}
+	}
+}
+
+URootMotionModifier_Warp::URootMotionModifier_Warp(const FObjectInitializer& ObjectInitializer)
+	: Super(ObjectInitializer)
+{
+}
+
+void URootMotionModifier_Warp::Update(const FMotionWarpingUpdateContext& Context)
+{
+	// Update playback times and state
+	Super::Update(Context);
+
+	// Cache sync point transform and trigger OnTargetTransformChanged if needed
+	const UMotionWarpingComponent* OwnerComp = GetOwnerComponent();
+	if (OwnerComp && GetState() == ERootMotionModifierState::Active)
+	{
+		const FMotionWarpingTarget* WarpTargetPtr = OwnerComp->FindWarpTarget(WarpTargetName);
+
+		// Disable if there is no target for us
+		if (WarpTargetPtr == nullptr)
+		{
+			UE_LOG(LogMotionWarping, Verbose, TEXT("MotionWarping: Marking RootMotionModifier as Disabled. Reason: Invalid Warp Target (%s). Char: %s Animation: %s [%f %f] [%f %f]"),
+				*WarpTargetName.ToString(), *GetNameSafe(OwnerComp->GetOwner()), *GetNameSafe(Animation.Get()), StartTime, EndTime, PreviousPosition, CurrentPosition);
+
+			SetState(ERootMotionModifierState::Disabled);
+			return;
+		}
+
+		// Get the warp point sent by the game
+		FTransform WarpPointTransformGame = WarpTargetPtr->GetTargetTrasform();
+
+		// Initialize our target transform (where the root should end at the end of the window) with the warp point sent by the game
+		FTransform TargetTransform = WarpPointTransformGame;
+
+		// Check if a warp point is defined in the animation. If so, we need to extract it and offset the target transform 
+		// the same amount the root bone is offset from the warp point in the animation
+		if (WarpPointAnimProvider != EWarpPointAnimProvider::None)
+		{
+			if (!CachedOffsetFromWarpPoint.IsSet())
+			{
+				if (const ACharacter* CharacterOwner = GetCharacterOwner())
+				{
+					if (WarpPointAnimProvider == EWarpPointAnimProvider::Static)
+					{
+						CachedOffsetFromWarpPoint = UMotionWarpingUtilities::CalculateRootTransformRelativeToWarpPointAtTime(*CharacterOwner, GetAnimation(), EndTime, WarpPointAnimTransform);
+					}
+					else if (WarpPointAnimProvider == EWarpPointAnimProvider::Bone)
+					{
+						CachedOffsetFromWarpPoint = UMotionWarpingUtilities::CalculateRootTransformRelativeToWarpPointAtTime(*CharacterOwner, GetAnimation(), EndTime, WarpPointAnimBoneName);
+					}
+				}
+			}
+
+			// Update Target Transform based on the offset between the root and the warp point in the animation
+			TargetTransform = CachedOffsetFromWarpPoint.GetValue() * WarpPointTransformGame;
+		}
+
+		if (!CachedTargetTransform.Equals(TargetTransform))
+		{
+			CachedTargetTransform = TargetTransform;
+
+			OnTargetTransformChanged();
+		}
+	}
+}
+
+void URootMotionModifier_Warp::OnTargetTransformChanged()
+{
+	if (const ACharacter* CharacterOwner = GetCharacterOwner())
+	{
+		ActualStartTime = PreviousPosition;
+
+		const float CapsuleHalfHeight = CharacterOwner->GetCapsuleComponent()->GetScaledCapsuleHalfHeight();
+		const FQuat CurrentRotation = CharacterOwner->GetActorQuat();
+		const FVector CurrentLocation = (CharacterOwner->GetActorLocation() - CurrentRotation.GetUpVector() * CapsuleHalfHeight);
+		StartTransform = FTransform(CurrentRotation, CurrentLocation);
+	}
+}
+
+FQuat URootMotionModifier_Warp::GetTargetRotation() const
+{
+	if (RotationType == EMotionWarpRotationType::Default)
+	{
+		return CachedTargetTransform.GetRotation();
+	}
+	else if (RotationType == EMotionWarpRotationType::Facing)
+	{
+		if (const ACharacter* CharacterOwner = GetCharacterOwner())
+		{
+			const FTransform& CharacterTransform = CharacterOwner->GetActorTransform();
+			const FVector ToSyncPoint = (CachedTargetTransform.GetLocation() - CharacterTransform.GetLocation()).GetSafeNormal2D();
+			return FRotationMatrix::MakeFromXZ(ToSyncPoint, FVector::UpVector).ToQuat();
+		}
+	}
+
+	return FQuat::Identity;
+}
+
+FQuat URootMotionModifier_Warp::WarpRotation(const FTransform& RootMotionDelta, const FTransform& RootMotionTotal, float DeltaSeconds)
+{
+	const ACharacter* CharacterOwner = GetCharacterOwner();
+	if (CharacterOwner == nullptr)
+	{
+		return FQuat::Identity;
+	}
+
+	const FQuat TotalRootMotionRotation = RootMotionTotal.GetRotation();
+	const FQuat CurrentRotation = CharacterOwner->GetActorQuat() * CharacterOwner->GetBaseRotationOffset();
+	const FQuat TargetRotation = CurrentRotation.Inverse() * (GetTargetRotation() * CharacterOwner->GetBaseRotationOffset());
+	const float TimeRemaining = (EndTime - PreviousPosition) * WarpRotationTimeMultiplier;
+	const float Alpha = FMath::Clamp(DeltaSeconds / TimeRemaining, 0.f, 1.f);
+	const FQuat TargetRotThisFrame = FQuat::Slerp(TotalRootMotionRotation, TargetRotation, Alpha);
+	const FQuat DeltaOut = TargetRotThisFrame * TotalRootMotionRotation.Inverse();
+	return (DeltaOut * RootMotionDelta.GetRotation());
+}
+
+#if !(UE_BUILD_SHIPPING || UE_BUILD_TEST)
+void URootMotionModifier_Warp::PrintLog(const FString& Name, const FTransform& OriginalRootMotion, const FTransform& WarpedRootMotion) const
+{
+	if (const ACharacter* CharacterOwner = GetCharacterOwner())
+	{
+		const float CapsuleHalfHeight = CharacterOwner->GetCapsuleComponent()->GetScaledCapsuleHalfHeight();
+		const FVector CurrentLocation = (CharacterOwner->GetActorLocation() - FVector(0.f, 0.f, CapsuleHalfHeight));
+		const FVector CurrentToTarget = (GetTargetLocation() - CurrentLocation).GetSafeNormal2D();
+		const FVector FutureLocation = CurrentLocation + (CharacterOwner->GetMesh()->ConvertLocalRootMotionToWorld(WarpedRootMotion)).GetTranslation();
+		const FRotator CurrentRotation = CharacterOwner->GetActorRotation();
+		const FRotator FutureRotation = (WarpedRootMotion.GetRotation() * CharacterOwner->GetActorQuat()).Rotator();
+		const float Dot = FVector::DotProduct(CharacterOwner->GetActorForwardVector(), CurrentToTarget);
+		const float CurrentDist2D = FVector::Dist2D(GetTargetLocation(), CurrentLocation);
+		const float FutureDist2D = FVector::Dist2D(GetTargetLocation(), FutureLocation);
+		const float DeltaSeconds = CharacterOwner->GetWorld()->GetDeltaSeconds();
+		const float Speed = WarpedRootMotion.GetTranslation().Size() / DeltaSeconds;
+		const float EndTimeOffset = CurrentPosition - EndTime;
+
+		UE_LOG(LogMotionWarping, Log, TEXT("%s NetMode: %d Char: %s Anim: %s Win: [%f %f][%f %f] DT: %f WT: %f ETOffset: %f Dist2D: %f Z: %f FDist2D: %f FZ: %f Dot: %f Delta: %s (%f) FDelta: %s (%f) Speed: %f Loc: %s FLoc: %s Rot: %s FRot: %s"),
+			*Name, (int32)CharacterOwner->GetWorld()->GetNetMode(), *GetNameSafe(CharacterOwner), *GetNameSafe(Animation.Get()), StartTime, EndTime, PreviousPosition, CurrentPosition, DeltaSeconds, CharacterOwner->GetWorld()->GetTimeSeconds(), EndTimeOffset,
+			CurrentDist2D, (GetTargetLocation().Z - CurrentLocation.Z), FutureDist2D, (GetTargetLocation().Z - FutureLocation.Z), Dot,
+			*OriginalRootMotion.GetTranslation().ToString(), OriginalRootMotion.GetTranslation().Size(), *WarpedRootMotion.GetTranslation().ToString(), WarpedRootMotion.GetTranslation().Size(), Speed,
+			*CurrentLocation.ToString(), *FutureLocation.ToString(), *CurrentRotation.ToCompactString(), *FutureRotation.ToCompactString());
+	}
+}
+#endif
+
+// URootMotionModifier_SimpleWarp
+///////////////////////////////////////////////////////////////
+
+UDEPRECATED_RootMotionModifier_SimpleWarp::UDEPRECATED_RootMotionModifier_SimpleWarp(const FObjectInitializer& ObjectInitializer)
+	: Super(ObjectInitializer)
+{
+}
+
+FTransform UDEPRECATED_RootMotionModifier_SimpleWarp::ProcessRootMotion(const FTransform& InRootMotion, float DeltaSeconds)
+{
+	const ACharacter* CharacterOwner = GetCharacterOwner();
+	if (CharacterOwner == nullptr)
+	{
+		return InRootMotion;
+	}
+
+	const FTransform& CharacterTransform = CharacterOwner->GetActorTransform();
+
+	FTransform FinalRootMotion = InRootMotion;
+
+	const FTransform RootMotionTotal = UMotionWarpingUtilities::ExtractRootMotionFromAnimation(Animation.Get(), PreviousPosition, EndTime);
+
+	if (bWarpTranslation)
+	{
+		FVector DeltaTranslation = InRootMotion.GetTranslation();
+
+		const FTransform RootMotionDelta = UMotionWarpingUtilities::ExtractRootMotionFromAnimation(Animation.Get(), PreviousPosition, FMath::Min(CurrentPosition, EndTime));
+
+		const float HorizontalDelta = RootMotionDelta.GetTranslation().Size2D();
+		const float HorizontalTarget = FVector::Dist2D(CharacterTransform.GetLocation(), GetTargetLocation());
+		const float HorizontalOriginal = RootMotionTotal.GetTranslation().Size2D();
+		const float HorizontalTranslationWarped = !FMath::IsNearlyZero(HorizontalOriginal) ? ((HorizontalDelta * HorizontalTarget) / HorizontalOriginal) : 0.f;
+
+		const FTransform MeshRelativeTransform = FTransform(CharacterOwner->GetBaseRotationOffset(), CharacterOwner->GetBaseTranslationOffset());
+		const FTransform MeshTransform = MeshRelativeTransform * CharacterOwner->GetActorTransform();
+		DeltaTranslation = MeshTransform.InverseTransformPositionNoScale(GetTargetLocation()).GetSafeNormal2D() * HorizontalTranslationWarped;
+
+		if (!bIgnoreZAxis)
+		{
+			const float CapsuleHalfHeight = CharacterOwner->GetCapsuleComponent()->GetScaledCapsuleHalfHeight();
+			const FVector CapsuleBottomLocation = (CharacterOwner->GetActorLocation() - FVector(0.f, 0.f, CapsuleHalfHeight));
+			const float VerticalDelta = RootMotionDelta.GetTranslation().Z;
+			const float VerticalTarget = GetTargetLocation().Z - CapsuleBottomLocation.Z;
+			const float VerticalOriginal = RootMotionTotal.GetTranslation().Z;
+			const float VerticalTranslationWarped = !FMath::IsNearlyZero(VerticalOriginal) ? ((VerticalDelta * VerticalTarget) / VerticalOriginal) : 0.f;
+
+			DeltaTranslation.Z = VerticalTranslationWarped;
+		}
+		else
+		{
+			DeltaTranslation.Z = InRootMotion.GetTranslation().Z;
+		}
+
+		FinalRootMotion.SetTranslation(DeltaTranslation);
+	}
+
+	if (bWarpRotation)
+	{
+		const FQuat WarpedRotation = WarpRotation(InRootMotion, RootMotionTotal, DeltaSeconds);
+		FinalRootMotion.SetRotation(WarpedRotation);
+	}
+
+	// Debug
+#if !(UE_BUILD_SHIPPING || UE_BUILD_TEST)
+	const int32 DebugLevel = FMotionWarpingCVars::CVarMotionWarpingDebug.GetValueOnGameThread();
+	if (DebugLevel == 1 || DebugLevel == 3)
+	{
+		PrintLog(TEXT("SimpleWarp"), InRootMotion, FinalRootMotion);
+	}
+
+	if (DebugLevel == 2 || DebugLevel == 3)
+	{
+		const float DrawDebugDuration = FMotionWarpingCVars::CVarMotionWarpingDrawDebugDuration.GetValueOnGameThread();
+		DrawDebugCoordinateSystem(CharacterOwner->GetWorld(), GetTargetLocation(), GetTargetRotator(), 50.f, false, DrawDebugDuration, 0, 1.f);
+	}
+#endif
+
+	return FinalRootMotion;
+}
+
+// URootMotionModifier_Scale
+///////////////////////////////////////////////////////////////
+
+URootMotionModifier_Scale* URootMotionModifier_Scale::AddRootMotionModifierScale(UMotionWarpingComponent* InMotionWarpingComp, const UAnimSequenceBase* InAnimation, float InStartTime, float InEndTime, FVector InScale)
+{
+	if (ensureAlways(InMotionWarpingComp))
+	{
+		URootMotionModifier_Scale* NewModifier = NewObject<URootMotionModifier_Scale>(InMotionWarpingComp);
+		NewModifier->Animation = InAnimation;
+		NewModifier->StartTime = InStartTime;
+		NewModifier->EndTime = InEndTime;
+		NewModifier->Scale = InScale;
+
+		InMotionWarpingComp->AddModifier(NewModifier);
+
+		return NewModifier;
+	}
+
+	return nullptr;
+}
\ No newline at end of file
diff --git a/Engine/Plugins/Animation/MotionWarping/Source/MotionWarping/Private/RootMotionModifier_AdjustmentBlendWarp.cpp b/Engine/Plugins/Animation/MotionWarping/Source/MotionWarping/Private/RootMotionModifier_AdjustmentBlendWarp.cpp
new file mode 100644
--- /dev/null
+++ b/Engine/Plugins/Animation/MotionWarping/Source/MotionWarping/Private/RootMotionModifier_AdjustmentBlendWarp.cpp
@@ -0,0 +1,465 @@
+// Copyright Epic Games, Inc. All Rights Reserved.
+
+#include "RootMotionModifier_AdjustmentBlendWarp.h"
+#include "Animation/AnimSequenceBase.h"
+#include "Animation/AnimMontage.h"
+#include "Animation/AnimInstance.h"
+#include "AnimationUtils.h"
+#include "AnimationRuntime.h"
+#include "DrawDebugHelpers.h"
+#include "GameFramework/Character.h"
+#include "GameFramework/CharacterMovementComponent.h"
+#include "MotionWarpingComponent.h"
+
+DECLARE_CYCLE_STAT(TEXT("MotionWarping PrecomputeWarpedTracks"), STAT_MotionWarping_PrecomputeWarpedTracks, STATGROUP_Anim);
+DECLARE_CYCLE_STAT(TEXT("MotionWarping ExtractMotionDelta"), STAT_MotionWarping_ExtractMotionDelta, STATGROUP_Anim);
+
+URootMotionModifier_AdjustmentBlendWarp::URootMotionModifier_AdjustmentBlendWarp(const FObjectInitializer& ObjectInitializer)
+	: Super(ObjectInitializer)
+{
+}
+
+void URootMotionModifier_AdjustmentBlendWarp::ExtractBoneTransformAtTime(FTransform& OutTransform, const FName& BoneName, float Time) const
+{
+	const int32 TrackIndex = Result.TrackNames.IndexOfByKey(BoneName);
+	ExtractBoneTransformAtTime(OutTransform, TrackIndex, Time);
+}
+
+void URootMotionModifier_AdjustmentBlendWarp::ExtractBoneTransformAtTime(FTransform& OutTransform, int32 TrackIndex, float Time) const
+{
+	if (!Result.AnimationTracks.IsValidIndex(TrackIndex))
+	{
+		OutTransform = FTransform::Identity;
+		return;
+	}
+
+	const int32 TotalFrames = FMath::Max(Result.AnimationTracks[TrackIndex].PosKeys.Num(), Result.AnimationTracks[TrackIndex].RotKeys.Num());
+	const float TrackLength = (EndTime - ActualStartTime);
+	const float TimePercent = (Time - ActualStartTime) / (EndTime - ActualStartTime);
+	const float RemappedTime = TimePercent * TrackLength;
+	FAnimationUtils::ExtractTransformFromTrack(RemappedTime, TotalFrames, TrackLength, Result.AnimationTracks[TrackIndex], EAnimInterpolationType::Linear, OutTransform);
+}
+
+void URootMotionModifier_AdjustmentBlendWarp::ExtractBoneTransformAtFrame(FTransform& OutTransform, int32 TrackIndex, int32 Frame) const
+{
+	if (!Result.AnimationTracks.IsValidIndex(TrackIndex) || !Result.AnimationTracks[TrackIndex].PosKeys.IsValidIndex(Frame))
+	{
+		OutTransform = FTransform::Identity;
+		return;
+	}
+
+	FAnimationUtils::ExtractTransformForFrameFromTrack(Result.AnimationTracks[TrackIndex], Frame, OutTransform);
+}
+
+void URootMotionModifier_AdjustmentBlendWarp::OnTargetTransformChanged()
+{
+	const ACharacter* CharacterOwner = GetCharacterOwner();
+	const USkeletalMeshComponent* SkelMeshComp = CharacterOwner ? CharacterOwner->GetMesh() : nullptr;
+
+	if (SkelMeshComp)
+	{
+		ActualStartTime = PreviousPosition;
+		CachedRootMotion = UMotionWarpingUtilities::ExtractRootMotionFromAnimation(Animation.Get(), ActualStartTime, EndTime);
+		CachedMeshRelativeTransform = SkelMeshComp->GetRelativeTransform();
+		CachedMeshTransform = SkelMeshComp->GetComponentTransform();
+
+		Result.AnimationTracks.Reset();
+		Result.TrackNames.Reset();
+	}
+}
+
+FTransform URootMotionModifier_AdjustmentBlendWarp::ProcessRootMotion(const FTransform& InRootMotion, float DeltaSeconds)
+{
+	// If warped tracks has not been generated yet, do it now
+	if (Result.GetNum() == 0)
+	{
+		PrecomputeWarpedTracks();
+	}
+
+	// Extract root motion from warped tracks
+	const FTransform FinalRootMotion = ExtractWarpedRootMotion();
+
+	// Debug
+#if !(UE_BUILD_SHIPPING || UE_BUILD_TEST)
+	const int32 DebugLevel = FMotionWarpingCVars::CVarMotionWarpingDebug.GetValueOnGameThread();
+	if (DebugLevel == 1 || DebugLevel == 3)
+	{
+		PrintLog(TEXT("AdjustmentBlendWarp"), InRootMotion, FinalRootMotion);
+	}
+
+	if (DebugLevel == 2 || DebugLevel == 3)
+	{
+		const float DrawDebugDuration = FMotionWarpingCVars::CVarMotionWarpingDrawDebugDuration.GetValueOnGameThread();
+		DrawDebugWarpedTracks(DrawDebugDuration);
+	}
+#endif
+
+	return FinalRootMotion;
+}
+
+void URootMotionModifier_AdjustmentBlendWarp::PrecomputeWarpedTracks()
+{
+	SCOPE_CYCLE_COUNTER(STAT_MotionWarping_PrecomputeWarpedTracks);
+
+	// First, extract pose at the end of the window for the bones we are going to warp
+
+	const ACharacter* CharacterOwner = GetCharacterOwner();
+	if (CharacterOwner == nullptr)
+	{
+		return;
+	}
+
+	const FBoneContainer& BoneContainer = CharacterOwner->GetMesh()->GetAnimInstance()->GetRequiredBones();
+
+	// Init FBoneContainer with only the bones that we are interested in
+	TArray<FBoneIndexType> RequiredBoneIndexArray;
+	RequiredBoneIndexArray.Add(0);
+
+	const bool bShouldWarpIKBones = bWarpIKBones && IKBones.Num() > 0;
+	if (bShouldWarpIKBones)
+	{
+		for (const FName& BoneName : IKBones)
+		{
+			const int32 BoneIndex = BoneContainer.GetPoseBoneIndexForBoneName(BoneName);
+			if (BoneIndex != INDEX_NONE)
+			{
+				RequiredBoneIndexArray.Add(BoneIndex);
+			}
+		}
+
+		BoneContainer.GetReferenceSkeleton().EnsureParentsExistAndSort(RequiredBoneIndexArray);
+	}
+
+	// Init BoneContainer
+	FBoneContainer RequiredBones(RequiredBoneIndexArray, FCurveEvaluationOption(false), *BoneContainer.GetAsset());
+
+	// Extract pose
+	FCSPose<FCompactPose> CSPose;
+	UMotionWarpingUtilities::ExtractComponentSpacePose(Animation.Get(), RequiredBones, EndTime, true, CSPose);
+
+	// Second, calculate additive pose
+
+	//Calculate additive translation for root bone
+	FVector RootTargetLocation = CachedMeshTransform.InverseTransformPositionNoScale(GetTargetLocation());
+
+	FVector RootTotalAdditiveTranslation = FVector::ZeroVector;
+	if (bWarpTranslation)
+	{
+		RootTotalAdditiveTranslation = RootTargetLocation - CachedRootMotion.GetLocation();
+
+		if (bIgnoreZAxis)
+		{
+			RootTotalAdditiveTranslation.Z = 0.f;
+		}
+	}
+
+	// Calculate additive rotation for root bone
+	FQuat RootTotalAdditiveRotation = FQuat::Identity;
+	if (bWarpRotation)
+	{
+		const FQuat TargetRotation = GetTargetRotation();
+		const FQuat OriginalRotation = CachedMeshRelativeTransform.GetRotation().Inverse() * (CachedRootMotion * CachedMeshTransform).GetRotation();
+		RootTotalAdditiveRotation = FQuat::FindBetweenNormals(OriginalRotation.GetForwardVector(), TargetRotation.GetForwardVector());
+	}
+
+	// Init Additive Pose
+	FCSPose<FCompactPose> AdditivePose;
+	AdditivePose.InitPose(&RequiredBones);
+	AdditivePose.SetComponentSpaceTransform(FCompactPoseBoneIndex(0), FTransform(RootTotalAdditiveRotation, RootTotalAdditiveTranslation));
+
+	// Calculate and add additive pose for IK bones
+	if (bShouldWarpIKBones)
+	{
+		const FTransform RootTargetPoseCS = FTransform((RootTotalAdditiveRotation * CachedRootMotion.GetRotation()), CachedRootMotion.GetTranslation() + RootTotalAdditiveTranslation);
+		for (int32 Idx = 1; Idx < CSPose.GetPose().GetNumBones(); Idx++)
+		{
+			const FName BoneName = RequiredBones.GetReferenceSkeleton().GetBoneName(RequiredBones.GetBoneIndicesArray()[Idx]);
+			if (IKBones.Contains(BoneName))
+			{
+				const int32 BoneIdx = Idx;
+				const FTransform BonePoseCS = CSPose.GetComponentSpaceTransform(FCompactPoseBoneIndex(BoneIdx));
+
+				const FTransform BoneTargetPoseCS = BonePoseCS * RootTargetPoseCS;
+				const FTransform BoneOriginalPoseCS = BonePoseCS * CachedRootMotion;
+
+				const FVector TotalAdditiveTranslation = BoneTargetPoseCS.GetLocation() - BoneOriginalPoseCS.GetLocation();
+
+				AdditivePose.SetComponentSpaceTransform(FCompactPoseBoneIndex(Idx), FTransform(TotalAdditiveTranslation));
+			}
+		}
+	}
+
+	// Finally, run adjustment blending to generate the warped poses for each bone
+
+	//@todo_fer: We could extract and cache this offline when the WarpingWindow is created
+	const float SampleRate = 1 / 60.f;
+	FMotionDeltaTrackContainer MotionDeltaTracks;
+	URootMotionModifier_AdjustmentBlendWarp::ExtractMotionDeltaFromRange(RequiredBones, Animation.Get(), ActualStartTime, EndTime, SampleRate, MotionDeltaTracks);
+
+	URootMotionModifier_AdjustmentBlendWarp::AdjustmentBlendWarp(RequiredBones, AdditivePose, MotionDeltaTracks, Result);
+}
+
+void URootMotionModifier_AdjustmentBlendWarp::ExtractMotionDeltaFromRange(const FBoneContainer& BoneContainer, const UAnimSequenceBase* Animation, float StartTime, float EndTime, float SampleRate, FMotionDeltaTrackContainer& OutMotionDeltaTracks)
+{
+	SCOPE_CYCLE_COUNTER(STAT_MotionWarping_ExtractMotionDelta);
+
+	check(Animation);
+
+	const int32 TotalFrames = FMath::CeilToInt((EndTime - StartTime) / SampleRate) + 1;
+	const int32 TotalBones = BoneContainer.GetCompactPoseNumBones();
+
+	FCompactPose FirstFramePose;
+	UMotionWarpingUtilities::ExtractLocalSpacePose(Animation, BoneContainer, StartTime, false, FirstFramePose);
+	const FTransform RootTransformFirstFrame = FirstFramePose[FCompactPoseBoneIndex(0)];
+
+	FCSPose<FCompactPose> CSPose;
+	UMotionWarpingUtilities::ExtractComponentSpacePose(Animation, BoneContainer, StartTime, false, CSPose);
+
+	OutMotionDeltaTracks.Init(TotalBones);
+	for (int32 BoneIdx = 0; BoneIdx < TotalBones; BoneIdx++)
+	{
+		const FTransform& BoneTransform = (CSPose.GetComponentSpaceTransform(FCompactPoseBoneIndex(BoneIdx)).GetRelativeTransform(RootTransformFirstFrame));
+
+		FMotionDeltaTrack& NewTrack = OutMotionDeltaTracks.Tracks.AddDefaulted_GetRef();
+
+		NewTrack.BoneTransformTrack.Reserve(TotalFrames);
+		NewTrack.DeltaTranslationTrack.Reserve(TotalFrames);
+		NewTrack.DeltaRotationTrack.Reserve(TotalFrames);
+
+		NewTrack.BoneTransformTrack.Add(BoneTransform);
+		NewTrack.DeltaTranslationTrack.Add(FVector::ZeroVector);
+		NewTrack.DeltaRotationTrack.Add(FRotator::ZeroRotator);
+
+		NewTrack.TotalTranslation = FVector::ZeroVector;
+		NewTrack.TotalRotation = FRotator::ZeroRotator;
+	}
+
+	float Time = StartTime;
+	while (Time < EndTime)
+	{
+		Time += FMath::Min(SampleRate, EndTime - Time);
+
+		UMotionWarpingUtilities::ExtractComponentSpacePose(Animation, BoneContainer, Time, false, CSPose);
+
+		for (int32 BoneIdx = 0; BoneIdx < TotalBones; BoneIdx++)
+		{
+			FMotionDeltaTrack& Track = OutMotionDeltaTracks.Tracks[BoneIdx];
+
+			const FTransform& BoneTransform = (CSPose.GetComponentSpaceTransform(FCompactPoseBoneIndex(BoneIdx)).GetRelativeTransform(RootTransformFirstFrame));
+
+			const FTransform& LastBoneTransform = Track.BoneTransformTrack.Last();
+
+			Track.BoneTransformTrack.Add(BoneTransform);
+
+			const FVector Translation = BoneTransform.GetTranslation();
+			const FVector LastTranslation = LastBoneTransform.GetTranslation();
+			const FVector DeltaTranslation(FMath::Abs(Translation.X - LastTranslation.X), FMath::Abs(Translation.Y - LastTranslation.Y), FMath::Abs(Translation.Z - LastTranslation.Z));
+			Track.DeltaTranslationTrack.Add(DeltaTranslation);
+			Track.TotalTranslation = Track.TotalTranslation + DeltaTranslation;
+
+			const FRotator Rotation = BoneTransform.GetRotation().Rotator();
+			const FRotator LastRotation = LastBoneTransform.GetRotation().Rotator();
+			const FRotator DeltaRotation(FMath::Abs(Rotation.Pitch - LastRotation.Pitch), FMath::Abs(Rotation.Yaw - LastRotation.Yaw), FMath::Abs(Rotation.Roll - LastRotation.Roll));
+			Track.DeltaRotationTrack.Add(DeltaRotation);
+			Track.TotalRotation = Track.TotalRotation + DeltaRotation;
+		}
+	}
+}
+
+void URootMotionModifier_AdjustmentBlendWarp::AdjustmentBlendWarp(const FBoneContainer& BoneContainer, const FCSPose<FCompactPose>& AdditivePose, const FMotionDeltaTrackContainer& MotionDeltaTracks, FAnimSequenceTrackContainer& Output)
+{
+	auto CalculateAdditive = [](const FVector& Total, const FVector& Delta, const FVector& Additive, const FVector& PreviousAdditive, float Alpha)
+	{
+		FVector CurrentAdditive = FVector::ZeroVector;
+		for (int32 Idx = 0; Idx < 3; Idx++)
+		{
+			if (!FMath::IsNearlyZero(Total[Idx], (FVector::FReal)1.f))
+			{
+				const FVector::FReal Percent = Delta[Idx] / Total[Idx];
+				const FVector::FReal AdditiveDelta = FMath::Abs(Additive[Idx]) * Percent;
+				CurrentAdditive[Idx] = (Additive[Idx] > 0.f) ? PreviousAdditive[Idx] + AdditiveDelta : PreviousAdditive[Idx] - AdditiveDelta;
+			}
+			else
+			{
+				CurrentAdditive[Idx] = Additive[Idx] * Alpha;
+			}
+		}
+		return CurrentAdditive;
+	};
+
+	Output.Initialize(BoneContainer.GetCompactPoseNumBones());
+
+	for (const FCompactPoseBoneIndex PoseBoneIndex : AdditivePose.GetPose().ForEachBoneIndex())
+	{
+		const FTransform& AdditiveTransform = AdditivePose.GetPose()[PoseBoneIndex];
+		if (AdditiveTransform.Equals(FTransform::Identity))
+		{
+			continue;
+		}
+
+		const int32 BoneIndex = PoseBoneIndex.GetInt();
+		const FMotionDeltaTrack& MotionDeltaTrack = MotionDeltaTracks.Tracks[BoneIndex];
+		const int32 TotalFrames = MotionDeltaTrack.BoneTransformTrack.Num();
+
+		FRawAnimSequenceTrack& Track = Output.AnimationTracks[BoneIndex];
+
+		Track.PosKeys.Reserve(TotalFrames);
+		Track.RotKeys.Reserve(TotalFrames);
+
+		Track.PosKeys.Add(FVector3f(MotionDeltaTrack.BoneTransformTrack[0].GetLocation()));
+		Track.RotKeys.Add(FQuat4f(MotionDeltaTrack.BoneTransformTrack[0].GetRotation()));
+
+		Track.ScaleKeys.Add(FVector3f(1));
+
+		Output.TrackNames[BoneIndex] = BoneContainer.GetReferenceSkeleton().GetBoneName(BoneContainer.GetBoneIndicesArray()[BoneIndex]);
+
+		const FVector TotalAdditiveTranslation = AdditiveTransform.GetTranslation();
+		const FVector TotalAdditiveRotation = AdditiveTransform.GetRotation().Rotator().Euler();
+
+		FVector PrevAdditiveTranslation = FVector::ZeroVector;
+		FVector PrevAdditiveRotation = FVector::ZeroVector;
+		for (int32 FrameIdx = 1; FrameIdx < TotalFrames; FrameIdx++)
+		{
+			const FTransform& BoneTransform = MotionDeltaTrack.BoneTransformTrack[FrameIdx];
+			const float Alpha = (FrameIdx / (float)(TotalFrames - 1));
+
+			// Translation Key
+			const FVector DeltaTranslation = MotionDeltaTrack.DeltaTranslationTrack[FrameIdx];
+			const FVector CurrentAdditiveTranslation = CalculateAdditive(MotionDeltaTrack.TotalTranslation, DeltaTranslation, TotalAdditiveTranslation, PrevAdditiveTranslation, Alpha);
+			Output.AnimationTracks[BoneIndex].PosKeys.Add(FVector3f(BoneTransform.GetTranslation() + CurrentAdditiveTranslation));
+			PrevAdditiveTranslation = CurrentAdditiveTranslation;
+
+			// Rotation Key
+			const FVector DeltaRotation = MotionDeltaTrack.DeltaRotationTrack[FrameIdx].Euler();
+			const FVector CurrentAdditiveRotation = CalculateAdditive(MotionDeltaTrack.TotalRotation.Euler(), DeltaRotation, TotalAdditiveRotation, PrevAdditiveRotation, Alpha);
+			Output.AnimationTracks[BoneIndex].RotKeys.Add(FQuat4f(FRotator::MakeFromEuler(BoneTransform.GetRotation().Rotator().Euler() + CurrentAdditiveRotation).Quaternion()));
+			PrevAdditiveRotation = CurrentAdditiveRotation;
+		}
+	}
+}
+
+FTransform URootMotionModifier_AdjustmentBlendWarp::ExtractWarpedRootMotion() const
+{
+	FTransform StartRootTransform;
+	ExtractBoneTransformAtTime(StartRootTransform, 0, PreviousPosition);
+
+	FTransform EndRootTransform;
+	ExtractBoneTransformAtTime(EndRootTransform, 0, CurrentPosition);
+
+	return EndRootTransform.GetRelativeTransform(StartRootTransform);
+}
+
+#if !(UE_BUILD_SHIPPING || UE_BUILD_TEST)
+void URootMotionModifier_AdjustmentBlendWarp::DrawDebugWarpedTracks(float DrawDuration) const
+{
+	const ACharacter* CharacterOwner = GetCharacterOwner();
+	const UWorld* World = CharacterOwner ? CharacterOwner->GetWorld() : nullptr;
+	if (World && Result.GetNum() > 0 && PreviousPosition <= EndTime)
+	{
+		FTransform RootStartTransform;
+		ExtractBoneTransformAtFrame(RootStartTransform, 0, 0);
+
+		for (int32 TrackIndex = 0; TrackIndex < Result.GetNum(); TrackIndex++)
+		{
+			const int32 TotalFrames = Result.AnimationTracks[TrackIndex].PosKeys.Num();
+			if (TotalFrames > 1)
+			{
+				for (int32 FrameIdx = 0; FrameIdx < Result.AnimationTracks[TrackIndex].PosKeys.Num() - 1; FrameIdx++)
+				{
+					FTransform T1;
+					ExtractBoneTransformAtFrame(T1, TrackIndex, FrameIdx);
+					T1 = T1 * RootStartTransform.Inverse() * CachedMeshTransform;
+
+					FTransform T2;
+					ExtractBoneTransformAtFrame(T2, TrackIndex, FrameIdx + 1);
+					T2 = T2 * RootStartTransform.Inverse() * CachedMeshTransform;
+
+					DrawDebugLine(World, T1.GetLocation(), T2.GetLocation(), FColor::Yellow, false, DrawDuration, 0, 0.5f);
+				}
+			}
+		}
+
+		const FTransform& MeshTransform = CharacterOwner->GetMesh()->GetComponentTransform();
+		DrawDebugCoordinateSystem(World, CachedMeshTransform.GetLocation(), CachedMeshTransform.Rotator(), 20.f, false, DrawDuration, 0, 1.f);
+		DrawDebugCoordinateSystem(World, MeshTransform.GetLocation(), MeshTransform.Rotator(), 20.f, false, DrawDuration, 0, 1.f);
+		DrawDebugCoordinateSystem(World, GetTargetLocation(), GetTargetRotator(), 50.f, false, DrawDuration, 0, 1.f);
+		DrawDebugCoordinateSystem(World, (CachedRootMotion * CachedMeshTransform).GetLocation(), (CachedRootMotion * CachedMeshTransform).Rotator(), 50.f, false, DrawDuration, 0, 1.f);
+	}
+}
+#endif
+
+void URootMotionModifier_AdjustmentBlendWarp::GetIKBoneTransformAndAlpha(FName BoneName, FTransform& OutTransform, float& OutAlpha) const
+{
+	if (Result.GetNum() == 0 || !bWarpIKBones || !IKBones.Contains(BoneName))
+	{
+		OutTransform = FTransform::Identity;
+		OutAlpha = 0.f;
+		return;
+	}
+
+	FTransform RootPrevPosition;
+	ExtractBoneTransformAtTime(RootPrevPosition, 0, 0.f);
+
+	FTransform BoneTransform;
+	ExtractBoneTransformAtTime(BoneTransform, BoneName, PreviousPosition);
+
+	OutTransform = BoneTransform * RootPrevPosition.Inverse() * CachedMeshTransform;
+	OutAlpha = Weight;
+}
+
+void URootMotionModifier_AdjustmentBlendWarp::GetAdjustmentBlendIKBoneTransformAndAlpha(ACharacter* Character, FName BoneName, FTransform& OutTransform, float& OutAlpha)
+{
+	OutTransform = FTransform::Identity;
+	OutAlpha = 0.f;
+
+	const UMotionWarpingComponent* MotionWarpingComp = Character ? Character->FindComponentByClass<UMotionWarpingComponent>() : nullptr;
+	if (MotionWarpingComp)
+	{
+		for (URootMotionModifier* Modifier : MotionWarpingComp->GetModifiers())
+		{
+			if (Modifier && Modifier->GetState() == ERootMotionModifierState::Active)
+			{
+				if (URootMotionModifier_AdjustmentBlendWarp* AdjustmentBlendWarpModifier = Cast<URootMotionModifier_AdjustmentBlendWarp>(Modifier))
+				{
+					// We must check if the Animation for the modifier is still relevant because in RootMotionFromMontageOnlyMode the montage could be aborted 
+					// but the modifier will remain in the list until the next update
+					const FAnimMontageInstance* RootMotionMontageInstance = Character->GetRootMotionAnimMontageInstance();
+					const UAnimMontage* Montage = RootMotionMontageInstance ? ToRawPtr(RootMotionMontageInstance->Montage) : nullptr;
+					if (Modifier->Animation == Montage)
+					{
+						AdjustmentBlendWarpModifier->GetIKBoneTransformAndAlpha(BoneName, OutTransform, OutAlpha);
+						return;
+					}
+				}
+			}
+		}
+	}
+}
+
+URootMotionModifier_AdjustmentBlendWarp* URootMotionModifier_AdjustmentBlendWarp::AddRootMotionModifierAdjustmentBlendWarp(UMotionWarpingComponent* InMotionWarpingComp, const UAnimSequenceBase* InAnimation, float InStartTime, float InEndTime,
+	FName InWarpTargetName, EWarpPointAnimProvider InWarpPointAnimProvider, FTransform InWarpPointAnimTransform, FName InWarpPointAnimBoneName,
+	bool bInWarpTranslation, bool bInIgnoreZAxis, bool bInWarpRotation, bool bInWarpIKBones, const TArray<FName>& InIKBones)
+{
+	if (ensureAlways(InMotionWarpingComp))
+	{
+		URootMotionModifier_AdjustmentBlendWarp* NewModifier = NewObject<URootMotionModifier_AdjustmentBlendWarp>(InMotionWarpingComp);
+		NewModifier->Animation = InAnimation;
+		NewModifier->StartTime = InStartTime;
+		NewModifier->EndTime = InEndTime;
+		NewModifier->WarpTargetName = InWarpTargetName;
+		NewModifier->WarpPointAnimProvider = InWarpPointAnimProvider;
+		NewModifier->WarpPointAnimTransform = InWarpPointAnimTransform;
+		NewModifier->WarpPointAnimBoneName = InWarpPointAnimBoneName;
+		NewModifier->bWarpTranslation = bInWarpTranslation;
+		NewModifier->bIgnoreZAxis = bInIgnoreZAxis;
+		NewModifier->bWarpRotation = bInWarpRotation;
+		NewModifier->bWarpIKBones = bInWarpIKBones;
+		NewModifier->IKBones = InIKBones;
+
+		InMotionWarpingComp->AddModifier(NewModifier);
+
+		return NewModifier;
+	}
+
+	return nullptr;
+}
\ No newline at end of file
diff --git a/Engine/Plugins/Animation/MotionWarping/Source/MotionWarping/Private/RootMotionModifier_SkewWarp.cpp b/Engine/Plugins/Animation/MotionWarping/Source/MotionWarping/Private/RootMotionModifier_SkewWarp.cpp
new file mode 100644
--- /dev/null
+++ b/Engine/Plugins/Animation/MotionWarping/Source/MotionWarping/Private/RootMotionModifier_SkewWarp.cpp
@@ -0,0 +1,233 @@
+// Copyright Epic Games, Inc. All Rights Reserved.
+
+#include "RootMotionModifier_SkewWarp.h"
+#include "DrawDebugHelpers.h"
+#include "GameFramework/Character.h"
+#include "GameFramework/CharacterMovementComponent.h"
+#include "Components/SkeletalMeshComponent.h"
+#include "Components/CapsuleComponent.h"
+#include "MotionWarpingComponent.h"
+
+URootMotionModifier_SkewWarp::URootMotionModifier_SkewWarp(const FObjectInitializer& ObjectInitializer)
+	: Super(ObjectInitializer)
+{
+}
+
+FVector URootMotionModifier_SkewWarp::WarpTranslation(const FTransform& CurrentTransform, const FVector& DeltaTranslation, const FVector& TotalTranslation, const FVector& TargetLocation)
+{
+	if (!DeltaTranslation.IsNearlyZero())
+	{
+		const FQuat CurrentRotation = CurrentTransform.GetRotation();
+		const FVector CurrentLocation = CurrentTransform.GetLocation();
+		const FVector FutureLocation = CurrentLocation + TotalTranslation;
+		const FVector CurrentToWorldOffset = TargetLocation - CurrentLocation;
+		const FVector CurrentToRootOffset = FutureLocation - CurrentLocation;
+
+		// Create a matrix we can use to put everything into a space looking straight at RootMotionSyncPosition. "forward" should be the axis along which we want to scale. 
+		FVector ToRootNormalized = CurrentToRootOffset.GetSafeNormal();
+
+		float BestMatchDot = FMath::Abs(FVector::DotProduct(ToRootNormalized, CurrentRotation.GetAxisX()));
+		FMatrix ToRootSyncSpace = FRotationMatrix::MakeFromXZ(ToRootNormalized, CurrentRotation.GetAxisZ());
+
+		float ZDot = FMath::Abs(FVector::DotProduct(ToRootNormalized, CurrentRotation.GetAxisZ()));
+		if (ZDot > BestMatchDot)
+		{
+			ToRootSyncSpace = FRotationMatrix::MakeFromXZ(ToRootNormalized, CurrentRotation.GetAxisX());
+			BestMatchDot = ZDot;
+		}
+
+		float YDot = FMath::Abs(FVector::DotProduct(ToRootNormalized, CurrentRotation.GetAxisY()));
+		if (YDot > BestMatchDot)
+		{
+			ToRootSyncSpace = FRotationMatrix::MakeFromXZ(ToRootNormalized, CurrentRotation.GetAxisZ());
+		}
+
+		// Put everything into RootSyncSpace.
+		const FVector RootMotionInSyncSpace = ToRootSyncSpace.InverseTransformVector(DeltaTranslation);
+		const FVector CurrentToWorldSync = ToRootSyncSpace.InverseTransformVector(CurrentToWorldOffset);
+		const FVector CurrentToRootMotionSync = ToRootSyncSpace.InverseTransformVector(CurrentToRootOffset);
+
+		FVector CurrentToWorldSyncNorm = CurrentToWorldSync;
+		CurrentToWorldSyncNorm.Normalize();
+
+		FVector CurrentToRootMotionSyncNorm = CurrentToRootMotionSync;
+		CurrentToRootMotionSyncNorm.Normalize();
+
+		// Calculate skew Yaw Angle. 
+		FVector FlatToWorld = FVector(CurrentToWorldSyncNorm.X, CurrentToWorldSyncNorm.Y, 0.0f);
+		FlatToWorld.Normalize();
+		FVector FlatToRoot = FVector(CurrentToRootMotionSyncNorm.X, CurrentToRootMotionSyncNorm.Y, 0.0f);
+		FlatToRoot.Normalize();
+		float AngleAboutZ = FMath::Acos(FVector::DotProduct(FlatToWorld, FlatToRoot));
+		float AngleAboutZNorm = FMath::DegreesToRadians(FRotator::NormalizeAxis(FMath::RadiansToDegrees(AngleAboutZ)));
+		if (FlatToWorld.Y < 0.0f)
+		{
+			AngleAboutZNorm *= -1.0f;
+		}
+
+		// Calculate Skew Pitch Angle. 
+		FVector ToWorldNoY = FVector(CurrentToWorldSyncNorm.X, 0.0f, CurrentToWorldSyncNorm.Z);
+		ToWorldNoY.Normalize();
+		FVector ToRootNoY = FVector(CurrentToRootMotionSyncNorm.X, 0.0f, CurrentToRootMotionSyncNorm.Z);
+		ToRootNoY.Normalize();
+		const float AngleAboutY = FMath::Acos(FVector::DotProduct(ToWorldNoY, ToRootNoY));
+		float AngleAboutYNorm = FMath::DegreesToRadians(FRotator::NormalizeAxis(FMath::RadiansToDegrees(AngleAboutY)));
+		if (ToWorldNoY.Z < 0.0f)
+		{
+			AngleAboutYNorm *= -1.0f;
+		}
+
+		FVector SkewedRootMotion = FVector::ZeroVector;
+		float ProjectedScale = FVector::DotProduct(CurrentToWorldSync, CurrentToRootMotionSyncNorm) / CurrentToRootMotionSync.Size();
+		if (ProjectedScale != 0.0f)
+		{
+			FMatrix ScaleMatrix;
+			ScaleMatrix.SetIdentity();
+			ScaleMatrix.SetAxis(0, FVector(ProjectedScale, 0.0f, 0.0f));
+			ScaleMatrix.SetAxis(1, FVector(0.0f, 1.0f, 0.0f));
+			ScaleMatrix.SetAxis(2, FVector(0.0f, 0.0f, 1.0f));
+
+			FMatrix ShearXAlongYMatrix;
+			ShearXAlongYMatrix.SetIdentity();
+			ShearXAlongYMatrix.SetAxis(0, FVector(1.0f, FMath::Tan(AngleAboutZNorm), 0.0f));
+			ShearXAlongYMatrix.SetAxis(1, FVector(0.0f, 1.0f, 0.0f));
+			ShearXAlongYMatrix.SetAxis(2, FVector(0.0f, 0.0f, 1.0f));
+
+			FMatrix ShearXAlongZMatrix;
+			ShearXAlongZMatrix.SetIdentity();
+			ShearXAlongZMatrix.SetAxis(0, FVector(1.0f, 0.0f, FMath::Tan(AngleAboutYNorm)));
+			ShearXAlongZMatrix.SetAxis(1, FVector(0.0f, 1.0f, 0.0f));
+			ShearXAlongZMatrix.SetAxis(2, FVector(0.0f, 0.0f, 1.0f));
+
+			FMatrix ScaledSkewMatrix = ScaleMatrix * ShearXAlongYMatrix * ShearXAlongZMatrix;
+
+			// Skew and scale the Root motion. 
+			SkewedRootMotion = ScaledSkewMatrix.TransformVector(RootMotionInSyncSpace);
+		}
+		else if (!CurrentToRootMotionSync.IsZero() && !CurrentToWorldSync.IsZero() && !RootMotionInSyncSpace.IsZero())
+		{
+			// Figure out ratio between remaining Root and remaining World. Then project scaled length of current Root onto World.
+			const float Scale = CurrentToWorldSync.Size() / CurrentToRootMotionSync.Size();
+			const float StepTowardTarget = RootMotionInSyncSpace.ProjectOnTo(RootMotionInSyncSpace).Size();
+			SkewedRootMotion = CurrentToWorldSyncNorm * (Scale * StepTowardTarget);
+		}
+
+		// Put our result back in world space.  
+		return ToRootSyncSpace.TransformVector(SkewedRootMotion);
+	}
+
+	return FVector::ZeroVector;
+}
+
+FTransform URootMotionModifier_SkewWarp::ProcessRootMotion(const FTransform& InRootMotion, float DeltaSeconds)
+{
+	const ACharacter* CharacterOwner = GetCharacterOwner();
+	if(CharacterOwner == nullptr)
+	{
+		return InRootMotion;
+	}
+
+	FTransform FinalRootMotion = InRootMotion;
+
+	const FTransform RootMotionTotal = UMotionWarpingUtilities::ExtractRootMotionFromAnimation(Animation.Get(), PreviousPosition, EndTime);
+	const FTransform RootMotionDelta = UMotionWarpingUtilities::ExtractRootMotionFromAnimation(Animation.Get(), PreviousPosition, FMath::Min(CurrentPosition, EndTime));
+
+	FTransform ExtraRootMotion = FTransform::Identity;
+	if (CurrentPosition > EndTime)
+	{
+		ExtraRootMotion = UMotionWarpingUtilities::ExtractRootMotionFromAnimation(Animation.Get(), EndTime, CurrentPosition);
+	}
+
+	if (bWarpTranslation)
+	{
+		const float CapsuleHalfHeight = CharacterOwner->GetCapsuleComponent()->GetScaledCapsuleHalfHeight();
+		const FQuat CurrentRotation = CharacterOwner->GetActorQuat();
+		const FVector CurrentLocation = (CharacterOwner->GetActorLocation() - CurrentRotation.GetUpVector() * CapsuleHalfHeight);
+
+		const FVector DeltaTranslation = RootMotionDelta.GetLocation();
+		const FVector TotalTranslation = RootMotionTotal.GetLocation();
+
+		FVector TargetLocation = GetTargetLocation();
+		if (bIgnoreZAxis)
+		{
+			TargetLocation.Z = CurrentLocation.Z;
+		}
+
+		// if there is translation in the animation, warp it
+		if (!TotalTranslation.IsNearlyZero())
+		{
+			if (!DeltaTranslation.IsNearlyZero())
+			{
+				const FTransform MeshTransform = FTransform(CharacterOwner->GetBaseRotationOffset(), CharacterOwner->GetBaseTranslationOffset()) * CharacterOwner->GetActorTransform();
+				TargetLocation = MeshTransform.InverseTransformPositionNoScale(TargetLocation);
+
+				const FVector WarpedTranslation = WarpTranslation(FTransform::Identity, DeltaTranslation, TotalTranslation, TargetLocation) + ExtraRootMotion.GetLocation();
+				FinalRootMotion.SetTranslation(WarpedTranslation);
+			}
+		}
+		// if there is no translation in the animation, add it
+		else
+		{
+			float Alpha = FMath::Clamp((CurrentPosition - ActualStartTime) / (EndTime - ActualStartTime), 0.f, 1.f);
+			Alpha = FAlphaBlend::AlphaToBlendOption(Alpha, AddTranslationEasingFunc, AddTranslationEasingCurve);
+
+			const FVector NextLocation = FMath::Lerp<FVector, float>(StartTransform.GetLocation(), TargetLocation, Alpha);
+			FVector FinalDeltaTranslation = (NextLocation - CurrentLocation);
+			FinalDeltaTranslation = (CurrentRotation.Inverse() * (TargetLocation - CurrentLocation).ToOrientationQuat()).GetForwardVector() * FinalDeltaTranslation.Size();
+			FinalDeltaTranslation = CharacterOwner->GetBaseRotationOffset().UnrotateVector(FinalDeltaTranslation);
+
+			FinalRootMotion.SetTranslation(FinalDeltaTranslation + ExtraRootMotion.GetLocation());
+		}
+	}
+
+	if(bWarpRotation)
+	{
+		const FQuat WarpedRotation = ExtraRootMotion.GetRotation() * WarpRotation(RootMotionDelta, RootMotionTotal, DeltaSeconds);
+		FinalRootMotion.SetRotation(WarpedRotation);
+	}
+
+	// Debug
+#if !(UE_BUILD_SHIPPING || UE_BUILD_TEST)
+	const int32 DebugLevel = FMotionWarpingCVars::CVarMotionWarpingDebug.GetValueOnGameThread();
+	if (DebugLevel == 1 || DebugLevel == 3)
+	{
+		PrintLog(TEXT("SkewWarp"), InRootMotion, FinalRootMotion);
+	}
+
+	if (DebugLevel == 2 || DebugLevel == 3)
+	{
+		const float DrawDebugDuration = FMotionWarpingCVars::CVarMotionWarpingDrawDebugDuration.GetValueOnGameThread();
+		DrawDebugCoordinateSystem(CharacterOwner->GetWorld(), GetTargetLocation(), GetTargetRotator(), 50.f, false, DrawDebugDuration, 0, 1.f);
+	}
+#endif
+
+	return FinalRootMotion;
+}
+
+URootMotionModifier_SkewWarp* URootMotionModifier_SkewWarp::AddRootMotionModifierSkewWarp(UMotionWarpingComponent* InMotionWarpingComp, const UAnimSequenceBase* InAnimation, float InStartTime, float InEndTime,
+	FName InWarpTargetName, EWarpPointAnimProvider InWarpPointAnimProvider, FTransform InWarpPointAnimTransform, FName InWarpPointAnimBoneName,
+	bool bInWarpTranslation, bool bInIgnoreZAxis, bool bInWarpRotation, EMotionWarpRotationType InRotationType, float InWarpRotationTimeMultiplier)
+{
+	if (ensureAlways(InMotionWarpingComp))
+	{
+		URootMotionModifier_SkewWarp* NewModifier = NewObject<URootMotionModifier_SkewWarp>(InMotionWarpingComp);
+		NewModifier->Animation = InAnimation;
+		NewModifier->StartTime = InStartTime;
+		NewModifier->EndTime = InEndTime;
+		NewModifier->WarpTargetName = InWarpTargetName;
+		NewModifier->WarpPointAnimProvider = InWarpPointAnimProvider;
+		NewModifier->WarpPointAnimTransform = InWarpPointAnimTransform;
+		NewModifier->WarpPointAnimBoneName = InWarpPointAnimBoneName;
+		NewModifier->bWarpTranslation = bInWarpTranslation;
+		NewModifier->bIgnoreZAxis = bInIgnoreZAxis;
+		NewModifier->bWarpRotation = bInWarpRotation;
+		NewModifier->RotationType = InRotationType;
+		NewModifier->WarpRotationTimeMultiplier = InWarpRotationTimeMultiplier;
+
+		InMotionWarpingComp->AddModifier(NewModifier);
+
+		return NewModifier;
+	}
+
+	return nullptr;
+}
\ No newline at end of file
diff --git a/Engine/Plugins/Animation/MotionWarping/Source/MotionWarping/Public/AnimNotifyState_MotionWarping.h b/Engine/Plugins/Animation/MotionWarping/Source/MotionWarping/Public/AnimNotifyState_MotionWarping.h
new file mode 100644
--- /dev/null
+++ b/Engine/Plugins/Animation/MotionWarping/Source/MotionWarping/Public/AnimNotifyState_MotionWarping.h
@@ -0,0 +1,56 @@
+// Copyright Epic Games, Inc. All Rights Reserved.
+
+#pragma once
+
+#include "CoreMinimal.h"
+#include "Animation/AnimNotifies/AnimNotifyState.h"
+#include "AnimNotifyState_MotionWarping.generated.h"
+
+class UMotionWarpingComponent;
+class UAnimSequenceBase;
+class URootMotionModifier;
+
+/** AnimNotifyState used to define a motion warping window in the animation */
+UCLASS(meta = (DisplayName = "Motion Warping"))
+class MOTIONWARPING_API UAnimNotifyState_MotionWarping : public UAnimNotifyState
+{
+	GENERATED_BODY()
+
+public:
+
+	//@TODO: Prevent notify callbacks and add comments explaining why we don't use those here.
+
+	UPROPERTY(EditAnywhere, Instanced, BlueprintReadWrite, Category = "Config")
+	TObjectPtr<URootMotionModifier> RootMotionModifier;
+
+	UAnimNotifyState_MotionWarping(const FObjectInitializer& ObjectInitializer);
+
+	/** Called from the MotionWarpingComp when this notify becomes relevant. See: UMotionWarpingComponent::Update */
+	void OnBecomeRelevant(UMotionWarpingComponent* MotionWarpingComp, const UAnimSequenceBase* Animation, float StartTime, float EndTime) const;
+
+	/** Creates a root motion modifier from the config class defined in the notify */
+	UFUNCTION(BlueprintNativeEvent, Category = "Motion Warping")
+	URootMotionModifier* AddRootMotionModifier(UMotionWarpingComponent* MotionWarpingComp, const UAnimSequenceBase* Animation, float StartTime, float EndTime) const;
+
+	UFUNCTION()
+	void OnRootMotionModifierActivate(UMotionWarpingComponent* MotionWarpingComp, URootMotionModifier* Modifier);
+
+	UFUNCTION()
+	void OnRootMotionModifierUpdate(UMotionWarpingComponent* MotionWarpingComp, URootMotionModifier* Modifier);
+
+	UFUNCTION()
+	void OnRootMotionModifierDeactivate(UMotionWarpingComponent* MotionWarpingComp, URootMotionModifier* Modifier);
+
+	UFUNCTION(BlueprintImplementableEvent, Category = "Motion Warping")
+	void OnWarpBegin(UMotionWarpingComponent* MotionWarpingComp, URootMotionModifier* Modifier) const;
+
+	UFUNCTION(BlueprintImplementableEvent, Category = "Motion Warping")
+	void OnWarpUpdate(UMotionWarpingComponent* MotionWarpingComp, URootMotionModifier* Modifier) const;
+
+	UFUNCTION(BlueprintImplementableEvent, Category = "Motion Warping")
+	void OnWarpEnd(UMotionWarpingComponent* MotionWarpingComp, URootMotionModifier* Modifier) const;
+
+#if WITH_EDITOR
+	virtual void ValidateAssociatedAssets() override;
+#endif
+};
\ No newline at end of file
diff --git a/Engine/Plugins/Animation/MotionWarping/Source/MotionWarping/Public/MotionWarping.h b/Engine/Plugins/Animation/MotionWarping/Source/MotionWarping/Public/MotionWarping.h
new file mode 100644
--- /dev/null
+++ b/Engine/Plugins/Animation/MotionWarping/Source/MotionWarping/Public/MotionWarping.h
@@ -0,0 +1,15 @@
+// Copyright Epic Games, Inc. All Rights Reserved.
+
+#pragma once
+
+#include "CoreMinimal.h"
+#include "Modules/ModuleManager.h"
+
+class FMotionWarpingModule : public IModuleInterface
+{
+public:
+
+	/** IModuleInterface implementation */
+	virtual void StartupModule() override;
+	virtual void ShutdownModule() override;
+};
diff --git a/Engine/Plugins/Animation/MotionWarping/Source/MotionWarping/Public/MotionWarpingComponent.h b/Engine/Plugins/Animation/MotionWarping/Source/MotionWarping/Public/MotionWarpingComponent.h
new file mode 100644
--- /dev/null
+++ b/Engine/Plugins/Animation/MotionWarping/Source/MotionWarping/Public/MotionWarpingComponent.h
@@ -0,0 +1,196 @@
+// Copyright Epic Games, Inc. All Rights Reserved.
+
+#pragma once
+
+#include "CoreMinimal.h"
+#include "Components/ActorComponent.h"
+#include "Kismet/KismetSystemLibrary.h"
+#include "RootMotionModifier.h"
+#include "MotionWarpingComponent.generated.h"
+
+class ACharacter;
+class UAnimSequenceBase;
+class UCharacterMovementComponent;
+class UMotionWarpingComponent;
+class UAnimNotifyState_MotionWarping;
+
+DECLARE_LOG_CATEGORY_EXTERN(LogMotionWarping, Log, All);
+
+#if !(UE_BUILD_SHIPPING || UE_BUILD_TEST)
+struct MOTIONWARPING_API FMotionWarpingCVars
+{
+	static TAutoConsoleVariable<int32> CVarMotionWarpingDisable;
+	static TAutoConsoleVariable<int32> CVarMotionWarpingDebug;
+	static TAutoConsoleVariable<float> CVarMotionWarpingDrawDebugDuration;
+};
+#endif
+
+USTRUCT(BlueprintType)
+struct FMotionWarpingWindowData
+{
+	GENERATED_BODY()
+
+	UPROPERTY(EditAnywhere, BlueprintReadWrite, Category = "Defaults")
+	TObjectPtr<UAnimNotifyState_MotionWarping> AnimNotify = nullptr;
+
+	UPROPERTY(EditAnywhere, BlueprintReadWrite, Category = "Defaults")
+	float StartTime = 0.f;
+
+	UPROPERTY(EditAnywhere, BlueprintReadWrite, Category = "Defaults")
+	float EndTime = 0.f;
+};
+
+UCLASS()
+class MOTIONWARPING_API UMotionWarpingUtilities : public UBlueprintFunctionLibrary
+{
+	GENERATED_BODY()
+
+public:
+
+	/** Extract bone pose in local space for all bones in BoneContainer. If Animation is a Montage the pose is extracted from the first track */
+	static void ExtractLocalSpacePose(const UAnimSequenceBase* Animation, const FBoneContainer& BoneContainer, float Time, bool bExtractRootMotion, FCompactPose& OutPose);
+
+	/** Extract bone pose in component space for all bones in BoneContainer. If Animation is a Montage the pose is extracted from the first track */
+	static void ExtractComponentSpacePose(const UAnimSequenceBase* Animation, const FBoneContainer& BoneContainer, float Time, bool bExtractRootMotion, FCSPose<FCompactPose>& OutPose);
+
+	/** Extract Root Motion transform from a contiguous position range */
+	UFUNCTION(BlueprintCallable, Category = "Motion Warping")
+	static FTransform ExtractRootMotionFromAnimation(const UAnimSequenceBase* Animation, float StartTime, float EndTime);
+
+	/** Extract root bone transform at a given time */
+	static FTransform ExtractRootTransformFromAnimation(const UAnimSequenceBase* Animation, float Time);
+
+	/** @return All the MotionWarping windows within the supplied animation */
+	UFUNCTION(BlueprintCallable, Category = "Motion Warping")
+	static void GetMotionWarpingWindowsFromAnimation(const UAnimSequenceBase* Animation, TArray<FMotionWarpingWindowData>& OutWindows);
+
+	/** @return All the MotionWarping windows within the supplied animation for a given Warp Target */
+	UFUNCTION(BlueprintCallable, Category = "Motion Warping")
+	static void GetMotionWarpingWindowsForWarpTargetFromAnimation(const UAnimSequenceBase* Animation, FName WarpTargetName, TArray<FMotionWarpingWindowData>& OutWindows);
+
+	/** @return root transform relative to the warp point bone at the supplied time */
+	static FTransform CalculateRootTransformRelativeToWarpPointAtTime(const ACharacter& Character, const UAnimSequenceBase* Animation, float Time, const FName& WarpPointBoneName);
+
+	/** @return root transform relative to the warp point transform at the supplied time */
+	static FTransform CalculateRootTransformRelativeToWarpPointAtTime(const ACharacter& Character, const UAnimSequenceBase* Animation, float Time, const FTransform& WarpPointTransform);
+};
+
+DECLARE_DYNAMIC_MULTICAST_DELEGATE_OneParam(FMotionWarpingPreUpdate, class UMotionWarpingComponent*, MotionWarpingComp);
+
+UCLASS(ClassGroup = Movement, meta = (BlueprintSpawnableComponent))
+class MOTIONWARPING_API UMotionWarpingComponent : public UActorComponent
+{
+	GENERATED_BODY()
+
+public:
+
+	/** Whether to look inside animations within montage when looking for warping windows */
+	UPROPERTY(EditAnywhere, BlueprintReadWrite, Category = "Config")
+	bool bSearchForWindowsInAnimsWithinMontages;
+
+	/** Event called before Root Motion Modifiers are updated */
+	UPROPERTY(BlueprintAssignable, Category = "Motion Warping")
+	FMotionWarpingPreUpdate OnPreUpdate;
+
+	UMotionWarpingComponent(const FObjectInitializer& ObjectInitializer);
+
+	virtual void InitializeComponent() override;
+	virtual void GetLifetimeReplicatedProps(TArray< FLifetimeProperty >& OutLifetimeProps) const;
+
+	/** Gets the character this component belongs to */
+	FORCEINLINE ACharacter* GetCharacterOwner() const { return CharacterOwner.Get(); }
+
+	/** Returns the list of root motion modifiers */
+	FORCEINLINE const TArray<URootMotionModifier*>& GetModifiers() const { return Modifiers; }
+
+	/** Check if we contain a RootMotionModifier for the supplied animation and time range */
+	bool ContainsModifier(const UAnimSequenceBase* Animation, float StartTime, float EndTime) const;
+
+	/** Add a new modifier */
+	int32 AddModifier(URootMotionModifier* Modifier);
+
+	/** Mark all the modifiers as Disable */
+	UFUNCTION(BlueprintCallable, Category = "Motion Warping")
+	void DisableAllRootMotionModifiers();
+
+	URootMotionModifier* AddModifierFromTemplate(URootMotionModifier* Template, const UAnimSequenceBase* Animation, float StartTime, float EndTime);
+
+	/** Find the target associated with a specified name */
+	FORCEINLINE const FMotionWarpingTarget* FindWarpTarget(const FName& WarpTargetName) const 
+	{ 
+		return WarpTargets.FindByPredicate([&WarpTargetName](const FMotionWarpingTarget& WarpTarget){ return WarpTarget.Name == WarpTargetName; });
+	}
+
+	/** Adds or update a warp target */
+	UFUNCTION(BlueprintCallable, Category = "Motion Warping")
+	void AddOrUpdateWarpTarget(const FMotionWarpingTarget& WarpTarget);
+
+	/** 
+	 * Create and adds or update a target associated with a specified name 
+	 * @param WarpTargetName Warp target identifier
+	 * @param TargetTransform Transform used to set the location and rotation for the warp target
+	 */
+	UFUNCTION(BlueprintCallable, Category = "Motion Warping")
+	void AddOrUpdateWarpTargetFromTransform(FName WarpTargetName, FTransform TargetTransform);
+
+	/**
+	 * Create and adds or update a target associated with a specified name
+	 * @param WarpTargetName Warp target identifier
+	 * @param Component Scene Component used to get the target transform
+	 * @param BoneName Optional bone or socket in the component used to get the target transform. 
+	 * @param bFollowComponent Whether the target transform should update while the warping is active. Useful for tracking moving targets.
+	 *		  Note that this will be one frame off if our owner ticks before the target actor. Add tick pre-requisites to avoid this.
+	 */
+	UFUNCTION(BlueprintCallable, Category = "Motion Warping")
+	void AddOrUpdateWarpTargetFromComponent(FName WarpTargetName, const USceneComponent* Component, FName BoneName, bool bFollowComponent);
+
+	/**
+	 * Create and adds or update a target associated with a specified name
+	 * @param WarpTargetName Warp target identifier
+	 * @param TargetLocation Location for the warp target
+	 */
+	UFUNCTION(BlueprintCallable, Category = "Motion Warping")
+	void AddOrUpdateWarpTargetFromLocation(FName WarpTargetName, FVector TargetLocation)
+	{
+		AddOrUpdateWarpTargetFromTransform(WarpTargetName, FTransform(TargetLocation));
+	}
+
+	/**
+	 * Create and adds or update a target associated with a specified name
+	 * @param WarpTargetName Warp target identifier
+	 * @param TargetLocation Location for the warp target
+	 */
+	UFUNCTION(BlueprintCallable, Category = "Motion Warping")
+	void AddOrUpdateWarpTargetFromLocationAndRotation(FName WarpTargetName, FVector TargetLocation, FRotator TargetRotation)
+	{
+		AddOrUpdateWarpTargetFromTransform(WarpTargetName, FTransform(TargetRotation, TargetLocation));
+	}
+
+	/** Removes the warp target associated with the specified key  */
+	UFUNCTION(BlueprintCallable, Category = "Motion Warping")
+	int32 RemoveWarpTarget(FName WarpTargetName);
+
+protected:
+
+	/** Character this component belongs to */
+	UPROPERTY(Transient)
+	TWeakObjectPtr<ACharacter> CharacterOwner;
+
+	/** List of root motion modifiers */
+	UPROPERTY(Transient)
+	TArray<TObjectPtr<URootMotionModifier>> Modifiers;
+
+	UPROPERTY(Transient, Replicated)
+	TArray<FMotionWarpingTarget> WarpTargets;
+
+#if !(UE_BUILD_SHIPPING || UE_BUILD_TEST)
+	TOptional<FVector> OriginalRootMotionAccum;
+	TOptional<FVector> WarpedRootMotionAccum;
+#endif
+
+	void Update(float DeltaSeconds);
+
+	bool FindAndUpdateWarpTarget(const FMotionWarpingTarget& WarpTarget);
+
+	FTransform ProcessRootMotionPreConvertToWorld(const FTransform& InRootMotion, class UCharacterMovementComponent* CharacterMovementComponent, float DeltaSeconds);
+};
\ No newline at end of file
diff --git a/Engine/Plugins/Animation/MotionWarping/Source/MotionWarping/Public/RootMotionModifier.h b/Engine/Plugins/Animation/MotionWarping/Source/MotionWarping/Public/RootMotionModifier.h
new file mode 100644
--- /dev/null
+++ b/Engine/Plugins/Animation/MotionWarping/Source/MotionWarping/Public/RootMotionModifier.h
@@ -0,0 +1,358 @@
+// Copyright Epic Games, Inc. All Rights Reserved.
+
+#pragma once
+
+#include "CoreMinimal.h"
+#include "Animation/AnimSequence.h"
+#include "AlphaBlend.h"
+#include "RootMotionModifier.generated.h"
+
+class UMotionWarpingComponent;
+class UAnimNotifyState_MotionWarping;
+class URootMotionModifier;
+class USceneComponent;
+
+/** 
+ * Context passed to any active root motion modifier during the update phase. 
+ * Contains relevant data from the animation that contributed to root motion this frame (or in the past when replaying saved moves)
+ */
+USTRUCT()
+struct FMotionWarpingUpdateContext
+{
+	GENERATED_BODY()
+
+	UPROPERTY()
+	TWeakObjectPtr<const UAnimSequenceBase> Animation = nullptr;
+
+	UPROPERTY()
+	float PreviousPosition = 0.f;
+
+	UPROPERTY()
+	float CurrentPosition = 0.f;
+
+	UPROPERTY()
+	float Weight = 0.f;
+
+	UPROPERTY()
+	float PlayRate = 1.f;
+
+	UPROPERTY()
+	float DeltaSeconds = 0.f;
+};
+
+/** The possible states of a Root Motion Modifier */
+UENUM(BlueprintType)
+enum class ERootMotionModifierState : uint8
+{
+	/** The modifier is waiting for the animation to hit the warping window */
+	Waiting,
+
+	/** The modifier is active and currently affecting the final root motion */
+	Active,
+
+	/** The modifier has been marked for removal. Usually because the warping window is done */
+	MarkedForRemoval,
+
+	/** The modifier will remain in the list (as long as the window is active) but will not modify the root motion */
+	Disabled
+};
+
+DECLARE_DYNAMIC_DELEGATE_TwoParams(FOnRootMotionModifierDelegate, UMotionWarpingComponent*, MotionWarpingComp, URootMotionModifier*, RootMotionModifier);
+
+// URootMotionModifier
+///////////////////////////////////////////////////////////////
+
+UCLASS(Abstract, BlueprintType, EditInlineNew)
+class MOTIONWARPING_API URootMotionModifier : public UObject
+{
+	GENERATED_BODY()
+
+public:
+
+	/** Source of the root motion we are warping */
+	UPROPERTY(BlueprintReadOnly, Category = "Defaults")
+	TWeakObjectPtr<const UAnimSequenceBase> Animation = nullptr;
+
+	/** Start time of the warping window */
+	UPROPERTY(BlueprintReadOnly, Category = "Defaults")
+	float StartTime = 0.f;
+
+	/** End time of the warping window */
+	UPROPERTY(BlueprintReadOnly, Category = "Defaults")
+	float EndTime = 0.f;
+
+	/** Previous playback time of the animation */
+	UPROPERTY(BlueprintReadOnly, Category = "Defaults")
+	float PreviousPosition = 0.f;
+
+	/** Current playback time of the animation */
+	UPROPERTY(BlueprintReadOnly, Category = "Defaults")
+	float CurrentPosition = 0.f;
+
+	/** Current blend weight of the animation */
+	UPROPERTY(BlueprintReadOnly, Category = "Defaults")
+	float Weight = 0.f;
+
+	/** Character owner transform at the time this modifier becomes active */
+	UPROPERTY(Transient, BlueprintReadOnly, Category = "Defaults")
+	FTransform StartTransform;
+
+	/** Actual playback time when the modifier becomes active */
+	UPROPERTY(Transient, BlueprintReadOnly, Category = "Defaults")
+	float ActualStartTime = 0.f;
+
+	/** Delegate called when this modifier is activated (starts affecting the root motion) */
+	UPROPERTY()
+	FOnRootMotionModifierDelegate OnActivateDelegate;
+
+	/** Delegate called when this modifier updates while active (affecting the root motion) */
+	UPROPERTY()
+	FOnRootMotionModifierDelegate OnUpdateDelegate;
+
+	/** Delegate called when this modifier is deactivated (stops affecting the root motion) */
+	UPROPERTY()
+	FOnRootMotionModifierDelegate OnDeactivateDelegate;
+
+	URootMotionModifier(const FObjectInitializer& ObjectInitializer);
+
+	/** Called when the state of the modifier changes */
+	virtual void OnStateChanged(ERootMotionModifierState LastState);
+
+	/** Sets the state of the modifier */
+	void SetState(ERootMotionModifierState NewState);
+
+	/** Returns the state of the modifier */
+	FORCEINLINE ERootMotionModifierState GetState() const { return State; }
+
+	/** Returns a pointer to the component that owns this modifier */
+	UMotionWarpingComponent* GetOwnerComponent() const;
+
+	/** Returns a pointer to the character that owns the component that owns this modifier */
+	class ACharacter* GetCharacterOwner() const;
+
+	virtual void Update(const FMotionWarpingUpdateContext& Context);
+	virtual FTransform ProcessRootMotion(const FTransform& InRootMotion, float DeltaSeconds) { return FTransform::Identity; }
+
+	FORCEINLINE const UAnimSequenceBase* GetAnimation() const { return Animation.Get(); }
+
+private:
+
+	friend UMotionWarpingComponent;
+
+	/** Current state */
+	UPROPERTY()
+	ERootMotionModifierState State = ERootMotionModifierState::Waiting;
+};
+
+// URootMotionModifier_Warp
+///////////////////////////////////////////////////////////////
+
+/** Represents a point of alignment in the world */
+USTRUCT(BlueprintType)
+struct MOTIONWARPING_API FMotionWarpingTarget
+{
+	GENERATED_BODY()
+
+	/** Unique name for this warp target */
+	UPROPERTY(EditAnywhere, BlueprintReadWrite, Category = "Defaults")
+	FName Name;
+
+	/** When the warp target is created from a component this stores the location of the component at the time of creation, otherwise its the location supplied by the user */
+	UPROPERTY(EditAnywhere, BlueprintReadWrite, Category = "Defaults")
+	FVector Location;
+
+	/** When the warp target is created from a component this stores the rotation of the component at the time of creation, otherwise its the rotation supplied by the user */
+	UPROPERTY(EditAnywhere, BlueprintReadWrite, Category = "Defaults")
+	FRotator Rotation;
+
+	/** Optional component used to calculate the final target transform */
+	UPROPERTY(EditAnywhere, BlueprintReadWrite, Category = "Defaults")
+	TWeakObjectPtr<const USceneComponent> Component;
+
+	/** Optional bone name in the component used to calculate the final target transform */
+	UPROPERTY(EditAnywhere, BlueprintReadWrite, Category = "Defaults")
+	FName BoneName;
+
+	/** Whether the target transform calculated from a component and an optional bone should be updated during the warp */
+	UPROPERTY(EditAnywhere, BlueprintReadWrite, Category = "Defaults")
+	bool bFollowComponent;
+
+	FMotionWarpingTarget()
+		: Name(NAME_None), Location(FVector::ZeroVector), Rotation(FRotator::ZeroRotator), Component(nullptr), BoneName(NAME_None), bFollowComponent(false) {}
+
+	FMotionWarpingTarget(const FName& InName, const FTransform& InTransform)
+		: Name(InName), Location(InTransform.GetLocation()), Rotation(InTransform.Rotator()), Component(nullptr), BoneName(NAME_None), bFollowComponent(false) {}
+
+	FMotionWarpingTarget(const FName& InName, const USceneComponent* InComp, FName InBoneName, bool bInbFollowComponent);
+
+	FTransform GetTargetTrasform() const;
+
+	FORCEINLINE FVector GetLocation() const { return GetTargetTrasform().GetLocation(); }
+	FORCEINLINE FQuat GetRotation() const { return GetTargetTrasform().GetRotation(); }
+	FORCEINLINE FRotator Rotator() const { return GetTargetTrasform().Rotator(); }
+
+	FORCEINLINE bool operator==(const FMotionWarpingTarget& Other) const
+	{
+		return Other.Name == Name && Other.Location.Equals(Location) && Other.Rotation.Equals(Rotation) && Other.Component == Component && Other.BoneName == BoneName && Other.bFollowComponent == bFollowComponent;
+	}
+
+	FORCEINLINE bool operator!=(const FMotionWarpingTarget& Other) const
+	{
+		return Other.Name != Name || !Other.Location.Equals(Location) || !Other.Rotation.Equals(Rotation) || Other.Component != Component || Other.BoneName != BoneName || Other.bFollowComponent != bFollowComponent;
+	}
+
+	static FTransform GetTargetTransformFromComponent(const USceneComponent* Comp, const FName& BoneName);
+
+};
+
+UENUM(BlueprintType)
+enum class EMotionWarpRotationType : uint8
+{
+	/** Character rotates to match the rotation of the sync point */
+	Default,
+
+	/** Character rotates to face the sync point */
+	Facing,
+};
+
+/** Method used to extract the warp point from the animation */
+UENUM(BlueprintType)
+enum class EWarpPointAnimProvider : uint8
+{
+	/** No warp point is provided */
+	None,
+
+	/** Warp point defined by a 'hard-coded' transform  user can enter through the warping notify */
+	Static,
+
+	/** Warp point defined by a bone */
+	Bone
+};
+
+UCLASS(Abstract)
+class MOTIONWARPING_API URootMotionModifier_Warp : public URootMotionModifier
+{
+	GENERATED_BODY()
+
+public:
+
+	/** Name used to find the warp target for this modifier */
+	UPROPERTY(EditAnywhere, BlueprintReadWrite, Category = "Config", meta = (ExposeOnSpawn))
+	FName WarpTargetName = NAME_None;
+
+	UPROPERTY(EditAnywhere, BlueprintReadWrite, Category = "Config")
+	EWarpPointAnimProvider WarpPointAnimProvider = EWarpPointAnimProvider::None;
+
+	//@TODO: Hide from the UI when Target != Static
+	UPROPERTY(EditAnywhere, BlueprintReadWrite, Category = "Config", meta = (EditCondition = "WarpPointAnimProvider == EWarpPointAnimProvider::Static"))
+	FTransform WarpPointAnimTransform = FTransform::Identity;
+
+	//@TODO: Hide from the UI when Target != Bone
+	UPROPERTY(EditAnywhere, BlueprintReadWrite, Category = "Config", meta = (EditCondition = "WarpPointAnimProvider == EWarpPointAnimProvider::Bone"))
+	FName WarpPointAnimBoneName = NAME_None;
+
+	/** Whether to warp the translation component of the root motion */
+	UPROPERTY(EditAnywhere, BlueprintReadWrite, Category = "Config")
+	bool bWarpTranslation = true;
+
+	/** Whether to ignore the Z component of the translation. Z motion will remain untouched */
+	UPROPERTY(EditAnywhere, BlueprintReadWrite, Category = "Config", meta = (EditCondition = "bWarpTranslation"))
+	bool bIgnoreZAxis = true;
+
+	/** Easing function used when adding translation. Only relevant when there is no translation in the animation */
+	UPROPERTY(EditAnywhere, BlueprintReadWrite, Category = "Config", meta = (EditCondition = "bWarpTranslation"))
+	EAlphaBlendOption AddTranslationEasingFunc = EAlphaBlendOption::Linear;
+
+	/** Custom curve used to add translation when there is none to warp. Only relevant when AddTranslationEasingFunc is set to Custom */
+	UPROPERTY(EditAnywhere, BlueprintReadWrite, Category = "Blend", meta = (EditCondition = "AddTranslationEasingFunc==EAlphaBlendOption::Custom", EditConditionHides))
+	TObjectPtr<class UCurveFloat> AddTranslationEasingCurve = nullptr;
+
+	/** Whether to warp the rotation component of the root motion */
+	UPROPERTY(EditAnywhere, BlueprintReadWrite, Category = "Config")
+	bool bWarpRotation = true;
+
+	/** Whether rotation should be warp to match the rotation of the sync point or to face the sync point */
+	UPROPERTY(EditAnywhere, BlueprintReadWrite, Category = "Config", meta = (EditCondition = "bWarpRotation"))
+	EMotionWarpRotationType RotationType;
+
+	/**
+	 * Allow to modify how fast the rotation is warped.
+	 * e.g if the window duration is 2sec and this is 0.5, the target rotation will be reached in 1sec instead of 2sec
+	 */
+	UPROPERTY(EditAnywhere, BlueprintReadWrite, Category = "Config", meta = (EditCondition = "bWarpRotation"))
+	float WarpRotationTimeMultiplier = 1.f;
+
+	URootMotionModifier_Warp(const FObjectInitializer& ObjectInitializer);
+
+	//~ Begin FRootMotionModifier Interface
+	virtual void Update(const FMotionWarpingUpdateContext& Context) override;
+	//~ End FRootMotionModifier Interface
+
+	/** Event called during update if the target transform changes while the warping is active */
+	virtual void OnTargetTransformChanged();
+
+#if !(UE_BUILD_SHIPPING || UE_BUILD_TEST)
+	void PrintLog(const FString& Name, const FTransform& OriginalRootMotion, const FTransform& WarpedRootMotion) const;
+#endif
+
+	FORCEINLINE FVector GetTargetLocation() const { return CachedTargetTransform.GetLocation(); }
+	FORCEINLINE FRotator GetTargetRotator() const { return GetTargetRotation().Rotator(); }
+	FQuat GetTargetRotation() const;
+
+	FQuat WarpRotation(const FTransform& RootMotionDelta, const FTransform& RootMotionTotal, float DeltaSeconds);
+
+protected:
+
+	UPROPERTY()
+	FTransform CachedTargetTransform = FTransform::Identity;
+
+	/** Cached of the offset from the warp target. Used to calculate the final target transform when a warp target is defined in the animation */
+	TOptional<FTransform> CachedOffsetFromWarpPoint;
+};
+
+// URootMotionModifier_SimpleWarp. 
+// DEPRECATED in favor of URootMotionModifier_SkewWarp (kept for reference)
+///////////////////////////////////////////////////////////////
+
+UCLASS(Deprecated, meta = (DisplayName = "Simple Warp"))
+class MOTIONWARPING_API UDEPRECATED_RootMotionModifier_SimpleWarp : public URootMotionModifier_Warp
+{
+	GENERATED_BODY()
+
+public:
+
+	UDEPRECATED_RootMotionModifier_SimpleWarp(const FObjectInitializer& ObjectInitializer);
+	virtual FTransform ProcessRootMotion(const FTransform& InRootMotion, float DeltaSeconds) override;
+};
+
+// URootMotionModifier_Scale
+///////////////////////////////////////////////////////////////
+
+UCLASS(meta = (DisplayName = "Scale"))
+class MOTIONWARPING_API URootMotionModifier_Scale : public URootMotionModifier
+{
+	GENERATED_BODY()
+
+public:
+
+	/** Vector used to scale each component of the translation */
+	UPROPERTY(EditAnywhere, BlueprintReadWrite, Category = "Config")
+	FVector Scale = FVector(1.f);
+
+	URootMotionModifier_Scale(const FObjectInitializer& ObjectInitializer) : Super(ObjectInitializer) {}
+
+	virtual FTransform ProcessRootMotion(const FTransform& InRootMotion, float DeltaSeconds) override
+	{
+		FTransform FinalRootMotion = InRootMotion;
+		FinalRootMotion.ScaleTranslation(Scale);
+		return FinalRootMotion;
+	}
+
+	UFUNCTION(BlueprintCallable, Category = "Motion Warping")
+	static URootMotionModifier_Scale* AddRootMotionModifierScale(
+		UPARAM(DisplayName = "Motion Warping Comp") UMotionWarpingComponent* InMotionWarpingComp,
+		UPARAM(DisplayName = "Animation") const UAnimSequenceBase* InAnimation,
+		UPARAM(DisplayName = "Start Time") float InStartTime,
+		UPARAM(DisplayName = "End Time") float InEndTime,
+		UPARAM(DisplayName = "Scale") FVector InScale = FVector(1.f));
+};
\ No newline at end of file
diff --git a/Engine/Plugins/Animation/MotionWarping/Source/MotionWarping/Public/RootMotionModifier_AdjustmentBlendWarp.h b/Engine/Plugins/Animation/MotionWarping/Source/MotionWarping/Public/RootMotionModifier_AdjustmentBlendWarp.h
new file mode 100644
--- /dev/null
+++ b/Engine/Plugins/Animation/MotionWarping/Source/MotionWarping/Public/RootMotionModifier_AdjustmentBlendWarp.h
@@ -0,0 +1,102 @@
+// Copyright Epic Games, Inc. All Rights Reserved.
+
+#pragma once
+
+#include "CoreMinimal.h"
+#include "RootMotionModifier.h"
+#include "RootMotionModifier_AdjustmentBlendWarp.generated.h"
+
+USTRUCT()
+struct FMotionDeltaTrack
+{
+	GENERATED_BODY()
+
+	UPROPERTY()
+	TArray<FTransform> BoneTransformTrack;
+
+	UPROPERTY()
+	TArray<FVector> DeltaTranslationTrack;
+
+	UPROPERTY()
+	TArray<FRotator> DeltaRotationTrack;
+
+	UPROPERTY()
+	FVector TotalTranslation = FVector::ZeroVector;
+
+	UPROPERTY()
+	FRotator TotalRotation = FRotator::ZeroRotator;
+};
+
+USTRUCT()
+struct FMotionDeltaTrackContainer
+{
+	GENERATED_BODY()
+
+	UPROPERTY()
+	TArray<FMotionDeltaTrack> Tracks;
+
+	void Init(int32 InNumTracks)
+	{
+		Tracks.Reserve(InNumTracks);
+	}
+};
+
+// EXPERIMENTAL: Marked with 'hidedropdown' to prevent it from showing up in the UI since it is not ready for production.
+UCLASS(hidedropdown, meta = (DisplayName = "Adjustment Blend Warp"))
+class MOTIONWARPING_API URootMotionModifier_AdjustmentBlendWarp : public URootMotionModifier_Warp
+{
+	GENERATED_BODY()
+
+public:
+
+	UPROPERTY(EditAnywhere, BlueprintReadWrite, Category = "Config")
+	bool bWarpIKBones = false;
+
+	UPROPERTY(EditAnywhere, BlueprintReadWrite, Category = "Config")
+	TArray<FName> IKBones;
+
+	URootMotionModifier_AdjustmentBlendWarp(const FObjectInitializer& ObjectInitializer);
+
+	virtual void OnTargetTransformChanged() override;
+	virtual FTransform ProcessRootMotion(const FTransform& InRootMotion, float DeltaSeconds) override;
+
+	void GetIKBoneTransformAndAlpha(FName BoneName, FTransform& OutTransform, float& OutAlpha) const;
+
+	//UFUNCTION(BlueprintCallable, Category = "Motion Warping")
+	static URootMotionModifier_AdjustmentBlendWarp* AddRootMotionModifierAdjustmentBlendWarp(UMotionWarpingComponent* InMotionWarpingComp, const UAnimSequenceBase* InAnimation, float InStartTime, float InEndTime,
+		FName InWarpTargetName, EWarpPointAnimProvider InWarpPointAnimProvider, FTransform InWarpPointAnimTransform, FName InWarpPointAnimBoneName,
+		bool bInWarpTranslation, bool bInIgnoreZAxis, bool bInWarpRotation, bool bInWarpIKBones, const TArray<FName>& InIKBones);
+
+	//UFUNCTION(BlueprintPure, Category = "Motion Warping")
+	static void GetAdjustmentBlendIKBoneTransformAndAlpha(ACharacter* Character, FName BoneName, FTransform& OutTransform, float& OutAlpha);
+
+protected:
+
+	UPROPERTY()
+	FTransform CachedMeshTransform;
+
+	UPROPERTY()
+	FTransform CachedMeshRelativeTransform;
+
+	UPROPERTY()
+	FTransform CachedRootMotion;
+
+	UPROPERTY()
+	FAnimSequenceTrackContainer Result;
+
+	void PrecomputeWarpedTracks();
+
+	FTransform ExtractWarpedRootMotion() const;
+
+	void ExtractBoneTransformAtTime(FTransform& OutTransform, const FName& BoneName, float Time) const;
+	void ExtractBoneTransformAtTime(FTransform& OutTransform, int32 TrackIndex, float Time) const;
+	void ExtractBoneTransformAtFrame(FTransform& OutTransform, int32 TrackIndex, int32 Frame) const;
+
+#if !(UE_BUILD_SHIPPING || UE_BUILD_TEST)
+	void DrawDebugWarpedTracks(float DrawDuration) const;
+#endif
+
+	static void ExtractMotionDeltaFromRange(const FBoneContainer& BoneContainer, const UAnimSequenceBase* Animation, float StartTime, float EndTime, float SampleRate, FMotionDeltaTrackContainer& OutMotionDeltaTracks);
+
+	static void AdjustmentBlendWarp(const FBoneContainer& BoneContainer, const FCSPose<FCompactPose>& AdditivePose, const FMotionDeltaTrackContainer& MotionDeltaTracks, FAnimSequenceTrackContainer& Output);
+};
\ No newline at end of file
diff --git a/Engine/Plugins/Animation/MotionWarping/Source/MotionWarping/Public/RootMotionModifier_SkewWarp.h b/Engine/Plugins/Animation/MotionWarping/Source/MotionWarping/Public/RootMotionModifier_SkewWarp.h
new file mode 100644
--- /dev/null
+++ b/Engine/Plugins/Animation/MotionWarping/Source/MotionWarping/Public/RootMotionModifier_SkewWarp.h
@@ -0,0 +1,37 @@
+// Copyright Epic Games, Inc. All Rights Reserved.
+
+#pragma once
+
+#include "CoreMinimal.h"
+#include "RootMotionModifier.h"
+#include "RootMotionModifier_SkewWarp.generated.h"
+
+UCLASS(meta = (DisplayName = "Skew Warp"))
+class MOTIONWARPING_API URootMotionModifier_SkewWarp : public URootMotionModifier_Warp
+{
+	GENERATED_BODY()
+
+public:
+
+	URootMotionModifier_SkewWarp(const FObjectInitializer& ObjectInitializer);
+
+	virtual FTransform ProcessRootMotion(const FTransform& InRootMotion, float DeltaSeconds) override;
+
+	static FVector WarpTranslation(const FTransform& CurrentTransform, const FVector& DeltaTranslation, const FVector& TotalTranslation, const FVector& TargetLocation);
+
+	UFUNCTION(BlueprintCallable, Category = "Motion Warping")
+	static URootMotionModifier_SkewWarp* AddRootMotionModifierSkewWarp(
+		UPARAM(DisplayName = "Motion Warping Comp") UMotionWarpingComponent* InMotionWarpingComp,
+		UPARAM(DisplayName = "Animation") const UAnimSequenceBase* InAnimation,
+		UPARAM(DisplayName = "Start Time") float InStartTime,
+		UPARAM(DisplayName = "End Time") float InEndTime,
+		UPARAM(DisplayName = "Warp Target Name") FName InWarpTargetName,
+		UPARAM(DisplayName = "Warp Point Anim Provider") EWarpPointAnimProvider InWarpPointAnimProvider,
+		UPARAM(DisplayName = "Warp Point Anim Transform") FTransform InWarpPointAnimTransform,
+		UPARAM(DisplayName = "Warp Point Anim Bone Name") FName InWarpPointAnimBoneName,
+		UPARAM(DisplayName = "Warp Translation") bool bInWarpTranslation = true,
+		UPARAM(DisplayName = "Ignore Z Axis") bool bInIgnoreZAxis = true,
+		UPARAM(DisplayName = "Warp Rotation") bool bInWarpRotation = true,
+		UPARAM(DisplayName = "Rotation Type") EMotionWarpRotationType InRotationType = EMotionWarpRotationType::Default,
+		UPARAM(DisplayName = "Warp Rotation Time Multiplier") float InWarpRotationTimeMultiplier = 1.f);
+};
\ No newline at end of file
diff --git a/Engine/Plugins/Experimental/Animation/MotionWarping/MotionWarping.uplugin b/Engine/Plugins/Experimental/Animation/MotionWarping/MotionWarping.uplugin
deleted file mode 100644
diff --git a/Engine/Plugins/Experimental/Animation/MotionWarping/Source/MotionWarping/MotionWarping.Build.cs b/Engine/Plugins/Experimental/Animation/MotionWarping/Source/MotionWarping/MotionWarping.Build.cs
deleted file mode 100644
diff --git a/Engine/Plugins/Experimental/Animation/MotionWarping/Source/MotionWarping/Private/AnimNotifyState_MotionWarping.cpp b/Engine/Plugins/Experimental/Animation/MotionWarping/Source/MotionWarping/Private/AnimNotifyState_MotionWarping.cpp
deleted file mode 100644
diff --git a/Engine/Plugins/Experimental/Animation/MotionWarping/Source/MotionWarping/Private/MotionWarping.cpp b/Engine/Plugins/Experimental/Animation/MotionWarping/Source/MotionWarping/Private/MotionWarping.cpp
deleted file mode 100644
diff --git a/Engine/Plugins/Experimental/Animation/MotionWarping/Source/MotionWarping/Private/MotionWarpingComponent.cpp b/Engine/Plugins/Experimental/Animation/MotionWarping/Source/MotionWarping/Private/MotionWarpingComponent.cpp
deleted file mode 100644
diff --git a/Engine/Plugins/Experimental/Animation/MotionWarping/Source/MotionWarping/Private/RootMotionModifier.cpp b/Engine/Plugins/Experimental/Animation/MotionWarping/Source/MotionWarping/Private/RootMotionModifier.cpp
deleted file mode 100644
diff --git a/Engine/Plugins/Experimental/Animation/MotionWarping/Source/MotionWarping/Private/RootMotionModifier_AdjustmentBlendWarp.cpp b/Engine/Plugins/Experimental/Animation/MotionWarping/Source/MotionWarping/Private/RootMotionModifier_AdjustmentBlendWarp.cpp
deleted file mode 100644
diff --git a/Engine/Plugins/Experimental/Animation/MotionWarping/Source/MotionWarping/Private/RootMotionModifier_SkewWarp.cpp b/Engine/Plugins/Experimental/Animation/MotionWarping/Source/MotionWarping/Private/RootMotionModifier_SkewWarp.cpp
deleted file mode 100644
diff --git a/Engine/Plugins/Experimental/Animation/MotionWarping/Source/MotionWarping/Public/AnimNotifyState_MotionWarping.h b/Engine/Plugins/Experimental/Animation/MotionWarping/Source/MotionWarping/Public/AnimNotifyState_MotionWarping.h
deleted file mode 100644
diff --git a/Engine/Plugins/Experimental/Animation/MotionWarping/Source/MotionWarping/Public/MotionWarping.h b/Engine/Plugins/Experimental/Animation/MotionWarping/Source/MotionWarping/Public/MotionWarping.h
deleted file mode 100644
diff --git a/Engine/Plugins/Experimental/Animation/MotionWarping/Source/MotionWarping/Public/MotionWarpingComponent.h b/Engine/Plugins/Experimental/Animation/MotionWarping/Source/MotionWarping/Public/MotionWarpingComponent.h
deleted file mode 100644
diff --git a/Engine/Plugins/Experimental/Animation/MotionWarping/Source/MotionWarping/Public/RootMotionModifier.h b/Engine/Plugins/Experimental/Animation/MotionWarping/Source/MotionWarping/Public/RootMotionModifier.h
deleted file mode 100644
diff --git a/Engine/Plugins/Experimental/Animation/MotionWarping/Source/MotionWarping/Public/RootMotionModifier_AdjustmentBlendWarp.h b/Engine/Plugins/Experimental/Animation/MotionWarping/Source/MotionWarping/Public/RootMotionModifier_AdjustmentBlendWarp.h
deleted file mode 100644
diff --git a/Engine/Plugins/Experimental/Animation/MotionWarping/Source/MotionWarping/Public/RootMotionModifier_SkewWarp.h b/Engine/Plugins/Experimental/Animation/MotionWarping/Source/MotionWarping/Public/RootMotionModifier_SkewWarp.h
deleted file mode 100644
