Marked Motion Warping plugin as beta

diff --git a/Engine/Plugins/Animation/MotionWarping/MotionWarping.uplugin b/Engine/Plugins/Animation/MotionWarping/MotionWarping.uplugin
--- a/Engine/Plugins/Animation/MotionWarping/MotionWarping.uplugin
+++ b/Engine/Plugins/Animation/MotionWarping/MotionWarping.uplugin
@@ -12,8 +12,8 @@
 	"SupportURL": "",
 	"EnabledByDefault": false,
 	"CanContainContent": true,
-	"IsBetaVersion": false,
-	"IsExperimentalVersion": true,
+	"IsBetaVersion": true,
+	"IsExperimentalVersion": false,
 	"Installed": false,
 	"Modules": [
 		{
