Motion Warping: Set SkewWarp as default root motion modifier for motion warping windows

diff --git a/Engine/Plugins/Animation/MotionWarping/Source/MotionWarping/Private/AnimNotifyState_MotionWarping.cpp b/Engine/Plugins/Animation/MotionWarping/Source/MotionWarping/Private/AnimNotifyState_MotionWarping.cpp
--- a/Engine/Plugins/Animation/MotionWarping/Source/MotionWarping/Private/AnimNotifyState_MotionWarping.cpp
+++ b/Engine/Plugins/Animation/MotionWarping/Source/MotionWarping/Private/AnimNotifyState_MotionWarping.cpp
@@ -3,7 +3,7 @@
 #include "AnimNotifyState_MotionWarping.h"
 #include "GameFramework/Actor.h"
 #include "MotionWarpingComponent.h"
-#include "RootMotionModifier.h"
+#include "RootMotionModifier_SkewWarp.h"
 
 #if WITH_EDITOR
 #include "Logging/MessageLog.h"
@@ -14,6 +14,7 @@
 UAnimNotifyState_MotionWarping::UAnimNotifyState_MotionWarping(const FObjectInitializer& ObjectInitializer)
 	: Super(ObjectInitializer)
 {
+	RootMotionModifier = ObjectInitializer.CreateDefaultSubobject<URootMotionModifier_SkewWarp>(this, TEXT("RootMotionModifier_SkewWarp"));
 }
 
 void UAnimNotifyState_MotionWarping::OnBecomeRelevant(UMotionWarpingComponent* MotionWarpingComp, const UAnimSequenceBase* Animation, float StartTime, float EndTime) const
