Fix potential instability in motion warping skew warp

diff --git a/Engine/Plugins/Animation/MotionWarping/Source/MotionWarping/Private/RootMotionModifier_SkewWarp.cpp b/Engine/Plugins/Animation/MotionWarping/Source/MotionWarping/Private/RootMotionModifier_SkewWarp.cpp
--- a/Engine/Plugins/Animation/MotionWarping/Source/MotionWarping/Private/RootMotionModifier_SkewWarp.cpp
+++ b/Engine/Plugins/Animation/MotionWarping/Source/MotionWarping/Private/RootMotionModifier_SkewWarp.cpp
@@ -170,15 +170,23 @@ FTransform URootMotionModifier_SkewWarp::ProcessRootMotion(const FTransform& InR
 		// if there is no translation in the animation, add it
 		else
 		{
-			float Alpha = FMath::Clamp((CurrentPosition - ActualStartTime) / (EndTime - ActualStartTime), 0.f, 1.f);
-			Alpha = FAlphaBlend::AlphaToBlendOption(Alpha, AddTranslationEasingFunc, AddTranslationEasingCurve);
+			const FVector DeltaToTarget = TargetLocation - CurrentLocation;
+			if (DeltaToTarget.IsNearlyZero())
+			{
+				FinalRootMotion.SetTranslation(FVector::ZeroVector);
+			}
+			else
+			{
+				float Alpha = FMath::Clamp((CurrentPosition - ActualStartTime) / (EndTime - ActualStartTime), 0.f, 1.f);
+				Alpha = FAlphaBlend::AlphaToBlendOption(Alpha, AddTranslationEasingFunc, AddTranslationEasingCurve);
 
-			const FVector NextLocation = FMath::Lerp<FVector, float>(StartTransform.GetLocation(), TargetLocation, Alpha);
-			FVector FinalDeltaTranslation = (NextLocation - CurrentLocation);
-			FinalDeltaTranslation = (CurrentRotation.Inverse() * (TargetLocation - CurrentLocation).ToOrientationQuat()).GetForwardVector() * FinalDeltaTranslation.Size();
-			FinalDeltaTranslation = CharacterOwner->GetBaseRotationOffset().UnrotateVector(FinalDeltaTranslation);
+				const FVector NextLocation = FMath::Lerp<FVector, float>(StartTransform.GetLocation(), TargetLocation, Alpha);
+				FVector FinalDeltaTranslation = (NextLocation - CurrentLocation);
+				FinalDeltaTranslation = (CurrentRotation.Inverse() * DeltaToTarget.ToOrientationQuat()).GetForwardVector() * FinalDeltaTranslation.Size();
+				FinalDeltaTranslation = CharacterOwner->GetBaseRotationOffset().UnrotateVector(FinalDeltaTranslation);
 
-			FinalRootMotion.SetTranslation(FinalDeltaTranslation + ExtraRootMotion.GetLocation());
+				FinalRootMotion.SetTranslation(FinalDeltaTranslation + ExtraRootMotion.GetLocation());
+			}
 		}
 	}
 
