[MotionWarping]
Remove the ensure from the CalculateRootTransformRelativeToWarpPointAtTime since we can get there by user error and it shouldn't make the system crash. We need to just notify the user so it could change the skeleton or fix the motion warping notify to use the right bone

#RNX

diff --git a/Engine/Plugins/Animation/MotionWarping/Source/MotionWarping/Private/MotionWarpingComponent.cpp b/Engine/Plugins/Animation/MotionWarping/Source/MotionWarping/Private/MotionWarpingComponent.cpp
--- a/Engine/Plugins/Animation/MotionWarping/Source/MotionWarping/Private/MotionWarpingComponent.cpp
+++ b/Engine/Plugins/Animation/MotionWarping/Source/MotionWarping/Private/MotionWarpingComponent.cpp
@@ -173,7 +173,7 @@ FTransform UMotionWarpingUtilities::CalculateRootTransformRelativeToWarpPointAtT
 		{
 			const FBoneContainer& FullBoneContainer = AnimInstance->GetRequiredBones();
 			const int32 BoneIndex = FullBoneContainer.GetPoseBoneIndexForBoneName(WarpPointBoneName);
-			if (ensure(BoneIndex != INDEX_NONE))
+			if (BoneIndex != INDEX_NONE)
 			{
 				TArray<FBoneIndexType> RequiredBoneIndexArray = { 0, (FBoneIndexType)BoneIndex };
 				FullBoneContainer.GetReferenceSkeleton().EnsureParentsExistAndSort(RequiredBoneIndexArray);
