Motion Warping: Added support for warp point bones that are not direct child of the root (from pull request #9837 on github)

diff --git a/Engine/Plugins/Animation/MotionWarping/Source/MotionWarping/Private/MotionWarpingComponent.cpp b/Engine/Plugins/Animation/MotionWarping/Source/MotionWarping/Private/MotionWarpingComponent.cpp
--- a/Engine/Plugins/Animation/MotionWarping/Source/MotionWarping/Private/MotionWarpingComponent.cpp
+++ b/Engine/Plugins/Animation/MotionWarping/Source/MotionWarping/Private/MotionWarpingComponent.cpp
@@ -187,7 +187,8 @@ FTransform UMotionWarpingUtilities::CalculateRootTransformRelativeToWarpPointAtT
 				const FTransform MeshCompRelativeRotInverse = FTransform(Character.GetBaseRotationOffset().Inverse());
 
 				const FTransform RootTransform = MeshCompRelativeRotInverse * Pose.GetComponentSpaceTransform(FCompactPoseBoneIndex(0));
-				const FTransform WarpPointTransform = MeshCompRelativeRotInverse * Pose.GetComponentSpaceTransform(FCompactPoseBoneIndex(1));
+				const FCompactPoseBoneIndex WarpPointBoneIndex = LimitedBoneContainer.GetCompactPoseIndexFromSkeletonPoseIndex(FSkeletonPoseBoneIndex(BoneIndex));
+				const FTransform WarpPointTransform = MeshCompRelativeRotInverse * Pose.GetComponentSpaceTransform(WarpPointBoneIndex);
 				return RootTransform.GetRelativeTransform(WarpPointTransform);
 			}
 		}
