Motion Warping: Rolling back change for adding support for warp point bones that are not direct child of the root (from pull request #9837 on github) since it is causing some tests to fail.

diff --git a/Engine/Plugins/Animation/MotionWarping/Source/MotionWarping/Private/MotionWarpingComponent.cpp b/Engine/Plugins/Animation/MotionWarping/Source/MotionWarping/Private/MotionWarpingComponent.cpp
--- a/Engine/Plugins/Animation/MotionWarping/Source/MotionWarping/Private/MotionWarpingComponent.cpp
+++ b/Engine/Plugins/Animation/MotionWarping/Source/MotionWarping/Private/MotionWarpingComponent.cpp
@@ -187,8 +187,7 @@ FTransform UMotionWarpingUtilities::CalculateRootTransformRelativeToWarpPointAtT
 				const FTransform MeshCompRelativeRotInverse = FTransform(Character.GetBaseRotationOffset().Inverse());
 
 				const FTransform RootTransform = MeshCompRelativeRotInverse * Pose.GetComponentSpaceTransform(FCompactPoseBoneIndex(0));
-				const FCompactPoseBoneIndex WarpPointBoneIndex = LimitedBoneContainer.GetCompactPoseIndexFromSkeletonPoseIndex(FSkeletonPoseBoneIndex(BoneIndex));
-				const FTransform WarpPointTransform = MeshCompRelativeRotInverse * Pose.GetComponentSpaceTransform(WarpPointBoneIndex);
+				const FTransform WarpPointTransform = MeshCompRelativeRotInverse * Pose.GetComponentSpaceTransform(FCompactPoseBoneIndex(1));
 				return RootTransform.GetRelativeTransform(WarpPointTransform);
 			}
 		}
