Added additional options to MotionWarping rotation functionality.
Optional ability to clamp rotation speed, and choose the method of rotation (lerp or constant)

diff --git a/Engine/Plugins/Animation/MotionWarping/Source/MotionWarping/Private/RootMotionModifier.cpp b/Engine/Plugins/Animation/MotionWarping/Source/MotionWarping/Private/RootMotionModifier.cpp
--- a/Engine/Plugins/Animation/MotionWarping/Source/MotionWarping/Private/RootMotionModifier.cpp
+++ b/Engine/Plugins/Animation/MotionWarping/Source/MotionWarping/Private/RootMotionModifier.cpp
@@ -305,8 +305,28 @@ FQuat URootMotionModifier_Warp::WarpRotation(const FTransform& RootMotionDelta,
 	const FQuat TargetRotation = CurrentRotation.Inverse() * (GetTargetRotation() * CharacterOwner->GetBaseRotationOffset());
 	const float TimeRemaining = (EndTime - PreviousPosition) * WarpRotationTimeMultiplier;
 	const float Alpha = FMath::Clamp(DeltaSeconds / TimeRemaining, 0.f, 1.f);
-	const FQuat TargetRotThisFrame = FQuat::Slerp(TotalRootMotionRotation, TargetRotation, Alpha);
+	FQuat TargetRotThisFrame = FQuat::Slerp(TotalRootMotionRotation, TargetRotation, Alpha);
+
+	if (RotationMethod != EMotionWarpRotationMethod::Slerp)
+	{
+		const float AngleDeltaThisFrame = TotalRootMotionRotation.AngularDistance(TargetRotThisFrame);
+		const float MaxAngleDelta = FMath::Abs(FMath::DegreesToRadians(DeltaSeconds * WarpMaxRotationRate));
+		const float TotalAngleDelta = TotalRootMotionRotation.AngularDistance(TargetRotation);
+		if (RotationMethod == EMotionWarpRotationMethod::ConstantRate && (TotalAngleDelta <= MaxAngleDelta))
+		{
+			TargetRotThisFrame = TargetRotation;
+		}
+		else if ((AngleDeltaThisFrame > MaxAngleDelta) || RotationMethod == EMotionWarpRotationMethod::ConstantRate)
+		{
+			const FVector CrossProduct = FVector::CrossProduct(TotalRootMotionRotation.Vector(), TargetRotation.Vector());
+			const float SignDirection = FMath::Sign(CrossProduct.Z);
+			const FQuat ClampedRotationThisFrame = FQuat(FVector(0, 0, 1), MaxAngleDelta * SignDirection);
+			TargetRotThisFrame = ClampedRotationThisFrame;
+		}
+	}
+
 	const FQuat DeltaOut = TargetRotThisFrame * TotalRootMotionRotation.Inverse();
+	
 	return (DeltaOut * RootMotionDelta.GetRotation());
 }
 
diff --git a/Engine/Plugins/Animation/MotionWarping/Source/MotionWarping/Private/RootMotionModifier_SkewWarp.cpp b/Engine/Plugins/Animation/MotionWarping/Source/MotionWarping/Private/RootMotionModifier_SkewWarp.cpp
--- a/Engine/Plugins/Animation/MotionWarping/Source/MotionWarping/Private/RootMotionModifier_SkewWarp.cpp
+++ b/Engine/Plugins/Animation/MotionWarping/Source/MotionWarping/Private/RootMotionModifier_SkewWarp.cpp
@@ -214,7 +214,7 @@ FTransform URootMotionModifier_SkewWarp::ProcessRootMotion(const FTransform& InR
 
 URootMotionModifier_SkewWarp* URootMotionModifier_SkewWarp::AddRootMotionModifierSkewWarp(UMotionWarpingComponent* InMotionWarpingComp, const UAnimSequenceBase* InAnimation, float InStartTime, float InEndTime,
 	FName InWarpTargetName, EWarpPointAnimProvider InWarpPointAnimProvider, FTransform InWarpPointAnimTransform, FName InWarpPointAnimBoneName,
-	bool bInWarpTranslation, bool bInIgnoreZAxis, bool bInWarpRotation, EMotionWarpRotationType InRotationType, float InWarpRotationTimeMultiplier)
+	bool bInWarpTranslation, bool bInIgnoreZAxis, bool bInWarpRotation, EMotionWarpRotationType InRotationType, EMotionWarpRotationMethod InRotationMethod, float InWarpRotationTimeMultiplier, float InWarpMaxRotationRate)
 {
 	if (ensureAlways(InMotionWarpingComp))
 	{
@@ -230,7 +230,9 @@ URootMotionModifier_SkewWarp* URootMotionModifier_SkewWarp::AddRootMotionModifie
 		NewModifier->bIgnoreZAxis = bInIgnoreZAxis;
 		NewModifier->bWarpRotation = bInWarpRotation;
 		NewModifier->RotationType = InRotationType;
+		NewModifier->RotationMethod = InRotationMethod;
 		NewModifier->WarpRotationTimeMultiplier = InWarpRotationTimeMultiplier;
+		NewModifier->WarpMaxRotationRate = InWarpMaxRotationRate;
 
 		InMotionWarpingComp->AddModifier(NewModifier);
 
diff --git a/Engine/Plugins/Animation/MotionWarping/Source/MotionWarping/Public/RootMotionModifier.h b/Engine/Plugins/Animation/MotionWarping/Source/MotionWarping/Public/RootMotionModifier.h
--- a/Engine/Plugins/Animation/MotionWarping/Source/MotionWarping/Public/RootMotionModifier.h
+++ b/Engine/Plugins/Animation/MotionWarping/Source/MotionWarping/Public/RootMotionModifier.h
@@ -214,6 +214,17 @@ enum class EMotionWarpRotationType : uint8
 	Facing,
 };
 
+UENUM(BlueprintType)
+enum class EMotionWarpRotationMethod : uint8
+{
+	/** Rotate with spherical linear interpolation */
+	Slerp, 
+	/** Rotate with spherical linear interpolation, not exceeding a max rotation rate*/
+	SlerpWithClampedRate,
+	/** Rotate with a constant rotation rate*/
+	ConstantRate,
+};
+
 /** Method used to extract the warp point from the animation */
 UENUM(BlueprintType)
 enum class EWarpPointAnimProvider : uint8
@@ -272,15 +283,23 @@ public:
 
 	/** Whether rotation should be warp to match the rotation of the sync point or to face the sync point */
 	UPROPERTY(EditAnywhere, BlueprintReadWrite, Category = "Config", meta = (EditCondition = "bWarpRotation"))
-	EMotionWarpRotationType RotationType;
+	EMotionWarpRotationType RotationType = EMotionWarpRotationType::Default;
+	
+	/** The method of rotation to use */
+	UPROPERTY(EditAnywhere, BlueprintReadWrite, Category = "Config", meta = (EditCondition = "bWarpRotation"))
+	EMotionWarpRotationMethod RotationMethod = EMotionWarpRotationMethod::Slerp;
 
 	/**
 	 * Allow to modify how fast the rotation is warped.
 	 * e.g if the window duration is 2sec and this is 0.5, the target rotation will be reached in 1sec instead of 2sec
 	 */
-	UPROPERTY(EditAnywhere, BlueprintReadWrite, Category = "Config", meta = (EditCondition = "bWarpRotation"))
+	UPROPERTY(EditAnywhere, BlueprintReadWrite, Category = "Config", meta = (EditCondition = "RotationMethod!=EMotionWarpRotationMethod::ConstantRate && bWarpRotation"))
 	float WarpRotationTimeMultiplier = 1.f;
 
+	/** Maximum rotation rate in degrees/sec. Will be the value used in constant rotation rate*/
+	UPROPERTY(EditAnywhere, BlueprintReadWrite, Category = "Config", meta = (EditCondition = "RotationMethod!=EMotionWarpRotationMethod::Slerp && bWarpRotation"))
+	float WarpMaxRotationRate = 0.f;
+
 	URootMotionModifier_Warp(const FObjectInitializer& ObjectInitializer);
 
 	//~ Begin FRootMotionModifier Interface
diff --git a/Engine/Plugins/Animation/MotionWarping/Source/MotionWarping/Public/RootMotionModifier_SkewWarp.h b/Engine/Plugins/Animation/MotionWarping/Source/MotionWarping/Public/RootMotionModifier_SkewWarp.h
--- a/Engine/Plugins/Animation/MotionWarping/Source/MotionWarping/Public/RootMotionModifier_SkewWarp.h
+++ b/Engine/Plugins/Animation/MotionWarping/Source/MotionWarping/Public/RootMotionModifier_SkewWarp.h
@@ -32,7 +32,9 @@ public:
 		UPARAM(DisplayName = "Ignore Z Axis") bool bInIgnoreZAxis = true,
 		UPARAM(DisplayName = "Warp Rotation") bool bInWarpRotation = true,
 		UPARAM(DisplayName = "Rotation Type") EMotionWarpRotationType InRotationType = EMotionWarpRotationType::Default,
-		UPARAM(DisplayName = "Warp Rotation Time Multiplier") float InWarpRotationTimeMultiplier = 1.f);
+		UPARAM(DisplayName = "Rotation Method") EMotionWarpRotationMethod InRotationMethod = EMotionWarpRotationMethod::Slerp,
+		UPARAM(DisplayName = "Warp Rotation Time Multiplier") float InWarpRotationTimeMultiplier = 1.f,
+		UPARAM(DisplayName = "Warp Max Rotation Rate") float InWarpMaxRotationRate = 0.f);
 };
 
 #if UE_ENABLE_INCLUDE_ORDER_DEPRECATED_IN_5_2
