Motion Warping: Added function to remove all warp targets

diff --git a/Engine/Plugins/Animation/MotionWarping/Source/MotionWarping/Private/MotionWarpingComponent.cpp b/Engine/Plugins/Animation/MotionWarping/Source/MotionWarping/Private/MotionWarpingComponent.cpp
--- a/Engine/Plugins/Animation/MotionWarping/Source/MotionWarping/Private/MotionWarpingComponent.cpp
+++ b/Engine/Plugins/Animation/MotionWarping/Source/MotionWarping/Private/MotionWarpingComponent.cpp
@@ -510,6 +510,20 @@ int32 UMotionWarpingComponent::RemoveWarpTarget(FName WarpTargetName)
 	return NumRemoved;
 }
 
+int32 UMotionWarpingComponent::RemoveAllWarpTargets()
+{
+	const int32 NumRemoved = WarpTargets.Num();
+
+	WarpTargets.Reset();
+
+	if (NumRemoved > 0)
+	{
+		MARK_PROPERTY_DIRTY_FROM_NAME(UMotionWarpingComponent, WarpTargets, this);
+	}
+
+	return NumRemoved;
+}
+
 void UMotionWarpingComponent::AddOrUpdateWarpTargetFromTransform(FName WarpTargetName, FTransform TargetTransform)
 {
 	AddOrUpdateWarpTarget(FMotionWarpingTarget(WarpTargetName, TargetTransform));
diff --git a/Engine/Plugins/Animation/MotionWarping/Source/MotionWarping/Public/MotionWarpingComponent.h b/Engine/Plugins/Animation/MotionWarping/Source/MotionWarping/Public/MotionWarpingComponent.h
--- a/Engine/Plugins/Animation/MotionWarping/Source/MotionWarping/Public/MotionWarpingComponent.h
+++ b/Engine/Plugins/Animation/MotionWarping/Source/MotionWarping/Public/MotionWarpingComponent.h
@@ -172,6 +172,10 @@ public:
 	UFUNCTION(BlueprintCallable, Category = "Motion Warping")
 	int32 RemoveWarpTarget(FName WarpTargetName);
 
+	/** Removes all warp targets */
+	UFUNCTION(BlueprintCallable, Category = "Motion Warping")
+	int32 RemoveAllWarpTargets();
+
 protected:
 
 	/** Character this component belongs to */
