Motion Warping: Added optional static offset for warp targets created from a component

diff --git a/Engine/Plugins/Animation/MotionWarping/Source/MotionWarping/Private/MotionWarpingComponent.cpp b/Engine/Plugins/Animation/MotionWarping/Source/MotionWarping/Private/MotionWarpingComponent.cpp
--- a/Engine/Plugins/Animation/MotionWarping/Source/MotionWarping/Private/MotionWarpingComponent.cpp
+++ b/Engine/Plugins/Animation/MotionWarping/Source/MotionWarping/Private/MotionWarpingComponent.cpp
@@ -529,7 +529,7 @@ void UMotionWarpingComponent::AddOrUpdateWarpTargetFromTransform(FName WarpTarge
 	AddOrUpdateWarpTarget(FMotionWarpingTarget(WarpTargetName, TargetTransform));
 }
 
-void UMotionWarpingComponent::AddOrUpdateWarpTargetFromComponent(FName WarpTargetName, const USceneComponent* Component, FName BoneName, bool bFollowComponent)
+void UMotionWarpingComponent::AddOrUpdateWarpTargetFromComponent(FName WarpTargetName, const USceneComponent* Component, FName BoneName, bool bFollowComponent, FVector LocationOffset, FRotator RotationOffset)
 {
 	if (Component == nullptr)
 	{
@@ -537,7 +537,7 @@ void UMotionWarpingComponent::AddOrUpdateWarpTargetFromComponent(FName WarpTarge
 		return;
 	}
 
-	AddOrUpdateWarpTarget(FMotionWarpingTarget(WarpTargetName, Component, BoneName, bFollowComponent));
+	AddOrUpdateWarpTarget(FMotionWarpingTarget(WarpTargetName, Component, BoneName, bFollowComponent, LocationOffset, RotationOffset));
 }
 
 URootMotionModifier* UMotionWarpingComponent::AddModifierFromTemplate(URootMotionModifier* Template, const UAnimSequenceBase* Animation, float StartTime, float EndTime)
diff --git a/Engine/Plugins/Animation/MotionWarping/Source/MotionWarping/Private/RootMotionModifier.cpp b/Engine/Plugins/Animation/MotionWarping/Source/MotionWarping/Private/RootMotionModifier.cpp
--- a/Engine/Plugins/Animation/MotionWarping/Source/MotionWarping/Private/RootMotionModifier.cpp
+++ b/Engine/Plugins/Animation/MotionWarping/Source/MotionWarping/Private/RootMotionModifier.cpp
@@ -13,7 +13,7 @@
 // FMotionWarpingTarget
 ///////////////////////////////////////////////////////////////
 
-FMotionWarpingTarget::FMotionWarpingTarget(const FName& InName, const USceneComponent* InComp, FName InBoneName, bool bInFollowComponent)
+FMotionWarpingTarget::FMotionWarpingTarget(const FName& InName, const USceneComponent* InComp, FName InBoneName, bool bInFollowComponent, const FVector& InLocOffset, const FRotator& InRotOffset)
 {
 	if (ensure(InComp))
 	{
@@ -21,6 +21,8 @@ FMotionWarpingTarget::FMotionWarpingTarget(const FName& InName, const USceneComp
 		Component = InComp;
 		BoneName = InBoneName;
 		bFollowComponent = bInFollowComponent;
+		LocationOffset = InLocOffset;
+		RotationOffset = InRotOffset;
 
 		FTransform Transform = FTransform::Identity;
 		if (BoneName != NAME_None)
@@ -32,6 +34,8 @@ FMotionWarpingTarget::FMotionWarpingTarget(const FName& InName, const USceneComp
 			Transform = InComp->GetComponentTransform();
 		}
 
+		Transform = FTransform(RotationOffset, LocationOffset) * Transform;
+
 		Location = Transform.GetLocation();
 		Rotation = Transform.Rotator();
 	}
@@ -60,14 +64,17 @@ FTransform FMotionWarpingTarget::GetTargetTrasform() const
 {
 	if (Component.IsValid() && bFollowComponent)
 	{
+		FTransform Transform = FTransform::Identity;
 		if (BoneName != NAME_None)
 		{
-			return FMotionWarpingTarget::GetTargetTransformFromComponent(Component.Get(), BoneName);
+			Transform = FMotionWarpingTarget::GetTargetTransformFromComponent(Component.Get(), BoneName);
 		}
 		else
 		{
-			return Component->GetComponentTransform();
+			Transform = Component->GetComponentTransform();
 		}
+
+		return FTransform(RotationOffset, LocationOffset) * Transform;
 	}
 
 	return FTransform(Rotation, Location);
diff --git a/Engine/Plugins/Animation/MotionWarping/Source/MotionWarping/Public/MotionWarpingComponent.h b/Engine/Plugins/Animation/MotionWarping/Source/MotionWarping/Public/MotionWarpingComponent.h
--- a/Engine/Plugins/Animation/MotionWarping/Source/MotionWarping/Public/MotionWarpingComponent.h
+++ b/Engine/Plugins/Animation/MotionWarping/Source/MotionWarping/Public/MotionWarpingComponent.h
@@ -142,9 +142,11 @@ public:
 	 * @param BoneName Optional bone or socket in the component used to get the target transform. 
 	 * @param bFollowComponent Whether the target transform should update while the warping is active. Useful for tracking moving targets.
 	 *		  Note that this will be one frame off if our owner ticks before the target actor. Add tick pre-requisites to avoid this.
+	 * @param LocationOffset Optional translation offset to apply to the transform we get from the component
+	 * @param RotationOffset Optional rotation offset to apply to the transform we get from the component
 	 */
 	UFUNCTION(BlueprintCallable, Category = "Motion Warping")
-	void AddOrUpdateWarpTargetFromComponent(FName WarpTargetName, const USceneComponent* Component, FName BoneName, bool bFollowComponent);
+	void AddOrUpdateWarpTargetFromComponent(FName WarpTargetName, const USceneComponent* Component, FName BoneName, bool bFollowComponent, FVector LocationOffset = FVector::ZeroVector, FRotator RotationOffset = FRotator::ZeroRotator);
 
 	/**
 	 * Create and adds or update a target associated with a specified name
diff --git a/Engine/Plugins/Animation/MotionWarping/Source/MotionWarping/Public/RootMotionModifier.h b/Engine/Plugins/Animation/MotionWarping/Source/MotionWarping/Public/RootMotionModifier.h
--- a/Engine/Plugins/Animation/MotionWarping/Source/MotionWarping/Public/RootMotionModifier.h
+++ b/Engine/Plugins/Animation/MotionWarping/Source/MotionWarping/Public/RootMotionModifier.h
@@ -176,13 +176,21 @@ struct MOTIONWARPING_API FMotionWarpingTarget
 	UPROPERTY(EditAnywhere, BlueprintReadWrite, Category = "Defaults")
 	bool bFollowComponent;
 
+	/** Optional static location offset. Only relevant when the warp target is created from a component */
+	UPROPERTY(EditAnywhere, BlueprintReadWrite, Category = "Defaults")
+	FVector LocationOffset = FVector::ZeroVector;
+
+	/** Optional static rotation offset. Only relevant when the warp target is created from a component */
+	UPROPERTY(EditAnywhere, BlueprintReadWrite, Category = "Defaults")
+	FRotator RotationOffset = FRotator::ZeroRotator;
+
 	FMotionWarpingTarget()
 		: Name(NAME_None), Location(FVector::ZeroVector), Rotation(FRotator::ZeroRotator), Component(nullptr), BoneName(NAME_None), bFollowComponent(false) {}
 
 	FMotionWarpingTarget(const FName& InName, const FTransform& InTransform)
 		: Name(InName), Location(InTransform.GetLocation()), Rotation(InTransform.Rotator()), Component(nullptr), BoneName(NAME_None), bFollowComponent(false) {}
 
-	FMotionWarpingTarget(const FName& InName, const USceneComponent* InComp, FName InBoneName, bool bInbFollowComponent);
+	FMotionWarpingTarget(const FName& InName, const USceneComponent* InComp, FName InBoneName, bool bInbFollowComponent, const FVector& InLocOffset = FVector::ZeroVector, const FRotator& InRotOffset = FRotator::ZeroRotator);
 
 	FTransform GetTargetTrasform() const;
 
@@ -192,12 +200,14 @@ struct MOTIONWARPING_API FMotionWarpingTarget
 
 	FORCEINLINE bool operator==(const FMotionWarpingTarget& Other) const
 	{
-		return Other.Name == Name && Other.Location.Equals(Location) && Other.Rotation.Equals(Rotation) && Other.Component == Component && Other.BoneName == BoneName && Other.bFollowComponent == bFollowComponent;
+		return Other.Name == Name && Other.Location.Equals(Location) && Other.Rotation.Equals(Rotation) && Other.Component == Component && Other.BoneName == BoneName && Other.bFollowComponent == bFollowComponent && 
+			Other.LocationOffset.Equals(LocationOffset) && Other.RotationOffset.Equals(RotationOffset);
 	}
 
 	FORCEINLINE bool operator!=(const FMotionWarpingTarget& Other) const
 	{
-		return Other.Name != Name || !Other.Location.Equals(Location) || !Other.Rotation.Equals(Rotation) || Other.Component != Component || Other.BoneName != BoneName || Other.bFollowComponent != bFollowComponent;
+		return Other.Name != Name || !Other.Location.Equals(Location) || !Other.Rotation.Equals(Rotation) || Other.Component != Component || Other.BoneName != BoneName || Other.bFollowComponent != bFollowComponent || 
+			!Other.LocationOffset.Equals(LocationOffset) || !Other.RotationOffset.Equals(RotationOffset);
 	}
 
 	static FTransform GetTargetTransformFromComponent(const USceneComponent* Comp, const FName& BoneName);
