Added rendering for root motion trajectory and skew warp event
- Added viewport and canvas draw methods for anim notify, notify state, and root motion modifier
- Added drawing of root trajectory and events in FAnimationViewportClient
- Fixed transform when reseting RootMotionTransform in UDebugSkelMeshComponent
- Update RootMotionReferenceTransform to be always at the beginning of the anim sequence (even for montages), makes it easier to deal with in notifies

diff --git a/Engine/Plugins/Animation/MotionWarping/Source/MotionWarping/MotionWarping.Build.cs b/Engine/Plugins/Animation/MotionWarping/Source/MotionWarping/MotionWarping.Build.cs
--- a/Engine/Plugins/Animation/MotionWarping/Source/MotionWarping/MotionWarping.Build.cs
+++ b/Engine/Plugins/Animation/MotionWarping/Source/MotionWarping/MotionWarping.Build.cs
@@ -16,14 +16,23 @@ public class MotionWarping : ModuleRules
 				"Engine",
 				"NetCore"
 			}
-			);			
+			);
 		
 		PrivateDependencyModuleNames.AddRange(
 			new string[]
 			{
 				"Slate",
-				"SlateCore"
-            }
+				"SlateCore",
+			}
 			);
+		
+		if (Target.bBuildEditor)
+		{
+			PublicDependencyModuleNames.AddRange(new string []
+				{
+					"UnrealEd",
+					"AnimGraph",
+				});
+		}
 	}
 }
diff --git a/Engine/Plugins/Animation/MotionWarping/Source/MotionWarping/Private/AnimNotifyState_MotionWarping.cpp b/Engine/Plugins/Animation/MotionWarping/Source/MotionWarping/Private/AnimNotifyState_MotionWarping.cpp
--- a/Engine/Plugins/Animation/MotionWarping/Source/MotionWarping/Private/AnimNotifyState_MotionWarping.cpp
+++ b/Engine/Plugins/Animation/MotionWarping/Source/MotionWarping/Private/AnimNotifyState_MotionWarping.cpp
@@ -114,4 +114,25 @@ void UAnimNotifyState_MotionWarping::ValidateAssociatedAssets()
 		}
 	}
 }
+
+void UAnimNotifyState_MotionWarping::DrawInEditor(FPrimitiveDrawInterface* PDI, USkeletalMeshComponent* MeshComp, const UAnimSequenceBase* Animation, const FAnimNotifyEvent& NotifyEvent) const
+{
+	if (RootMotionModifier)
+	{
+		FAnimNotifyEvent NotifyEventWithColor = NotifyEvent;
+		NotifyEventWithColor.NotifyColor = NotifyColor;
+		RootMotionModifier->DrawInEditor(PDI, MeshComp, Animation, NotifyEventWithColor);
+	}
+}
+
+void UAnimNotifyState_MotionWarping::DrawCanvasInEditor(FCanvas& Canvas, FSceneView& View, USkeletalMeshComponent* MeshComp, const UAnimSequenceBase* Animation, const FAnimNotifyEvent& NotifyEvent) const
+{
+	if (RootMotionModifier)
+	{
+		FAnimNotifyEvent NotifyEventWithColor = NotifyEvent;
+		NotifyEventWithColor.NotifyColor = NotifyColor;
+		RootMotionModifier->DrawCanvasInEditor(Canvas, View, MeshComp, Animation, NotifyEventWithColor);
+	}
+}
+
 #endif
diff --git a/Engine/Plugins/Animation/MotionWarping/Source/MotionWarping/Private/RootMotionModifier_SkewWarp.cpp b/Engine/Plugins/Animation/MotionWarping/Source/MotionWarping/Private/RootMotionModifier_SkewWarp.cpp
--- a/Engine/Plugins/Animation/MotionWarping/Source/MotionWarping/Private/RootMotionModifier_SkewWarp.cpp
+++ b/Engine/Plugins/Animation/MotionWarping/Source/MotionWarping/Private/RootMotionModifier_SkewWarp.cpp
@@ -5,6 +5,16 @@
 #include "GameFramework/Character.h"
 #include "Components/CapsuleComponent.h"
 #include "MotionWarpingComponent.h"
+#if WITH_EDITOR
+#include "Animation/AnimInstance.h"
+#include "Animation/DebugSkelMeshComponent.h"
+#include "AnimPreviewInstance.h"
+#include "SceneManagement.h"
+#include "BonePose.h"
+#include "Engine/Font.h"
+#include "CanvasTypes.h"
+#include "Animation/AnimSequenceHelpers.h"
+#endif
 
 #include UE_INLINE_GENERATED_CPP_BY_NAME(RootMotionModifier_SkewWarp)
 
@@ -241,3 +251,153 @@ URootMotionModifier_SkewWarp* URootMotionModifier_SkewWarp::AddRootMotionModifie
 
 	return nullptr;
 }
+
+#if WITH_EDITOR
+
+FTransform URootMotionModifier_SkewWarp::GetDebugWarpPointTransform(USkeletalMeshComponent* MeshComp, const UAnimSequenceBase* InAnimation, const UMirrorDataTable* MirrorTable, const float NotifyEndTime) const
+{
+	FTransform WarpPointTransform = FTransform::Identity;
+	if (WarpPointAnimProvider == EWarpPointAnimProvider::None)
+	{
+		WarpPointTransform = UE::Anim::ExtractRootMotionFromAnimationAsset(InAnimation, MirrorTable, 0.0, NotifyEndTime);
+	}
+	else if (WarpPointAnimProvider == EWarpPointAnimProvider::Static)
+	{
+		WarpPointTransform = WarpPointAnimTransform;
+	}
+	else if (WarpPointAnimProvider == EWarpPointAnimProvider::Bone)
+	{
+		if (const UAnimInstance* AnimInstance = MeshComp->GetAnimInstance())
+		{
+			const FBoneContainer& FullBoneContainer = AnimInstance->GetRequiredBones();
+			const int32 BoneIndex = FullBoneContainer.GetPoseBoneIndexForBoneName(WarpPointAnimBoneName);
+			if (BoneIndex != INDEX_NONE)
+			{
+				TArray<FBoneIndexType> RequiredBoneIndexArray = { 0, (FBoneIndexType)BoneIndex };
+				FullBoneContainer.GetReferenceSkeleton().EnsureParentsExistAndSort(RequiredBoneIndexArray);
+
+				const FBoneContainer LimitedBoneContainer(RequiredBoneIndexArray, UE::Anim::FCurveFilterSettings(UE::Anim::ECurveFilterMode::DisallowAll), *FullBoneContainer.GetAsset());
+
+				FCSPose<FCompactPose> Pose;
+				UMotionWarpingUtilities::ExtractComponentSpacePose(InAnimation, LimitedBoneContainer, NotifyEndTime, false, Pose);
+
+				WarpPointTransform = Pose.GetComponentSpaceTransform(FCompactPoseBoneIndex(1));
+			}
+		}
+	}
+	return WarpPointTransform;
+}
+
+void URootMotionModifier_SkewWarp::DrawInEditor(FPrimitiveDrawInterface* PDI, USkeletalMeshComponent* MeshComp, const UAnimSequenceBase* InAnimation, const FAnimNotifyEvent& NotifyEvent) const
+{
+	constexpr float DepthBias = 4.0f;
+	constexpr bool bScreenSpace = true;
+
+	check(MeshComp);
+	check(PDI);
+
+	// Early out if the animation does not have root motion.
+	if (!InAnimation || !InAnimation->HasRootMotion())
+	{
+		return;
+	}
+	
+	const float NotifyStartTime = NotifyEvent.GetTriggerTime();
+	const float NotifyEndTime = NotifyEvent.GetEndTriggerTime();
+
+	const UMirrorDataTable* MirrorTable = nullptr;
+	FTransform ReferenceTransform = FTransform::Identity;
+	if (UDebugSkelMeshComponent* DebugMeshComp = Cast<UDebugSkelMeshComponent>(MeshComp))
+	{
+		ReferenceTransform = DebugMeshComp->RootMotionReferenceTransform; // Actors location at the beginning of the animation
+		MirrorTable = DebugMeshComp->PreviewInstance ? DebugMeshComp->PreviewInstance->GetMirrorDataTable() : nullptr;
+	}
+
+	const FTransform WarpPointTransform = GetDebugWarpPointTransform(MeshComp, InAnimation, MirrorTable, NotifyEndTime) * ReferenceTransform;
+
+	// Draw notify duration on root motion track.
+	const FFrameRate FrameRate = InAnimation->GetSamplingFrameRate();
+
+	const int32 StartFrame = FrameRate.AsFrameTime(NotifyStartTime).CeilToFrame().Value;
+	const int32 EndFrame = FrameRate.AsFrameTime(NotifyEndTime).FloorToFrame().Value;
+
+	const FTransform StartRootTransform = UE::Anim::ExtractRootMotionFromAnimationAsset(InAnimation, MirrorTable, 0.0, NotifyStartTime) * ReferenceTransform;
+	const FTransform EndRootTransform = UE::Anim::ExtractRootMotionFromAnimationAsset(InAnimation, MirrorTable, 0.0, NotifyEndTime) * ReferenceTransform;
+
+	constexpr double TrackOffset = 2.0; 
+	
+	FVector PrevLocation = StartRootTransform.GetLocation() + StartRootTransform.GetUnitAxis(EAxis::Z) * TrackOffset;
+	for (int32 Frame = StartFrame; Frame <= EndFrame; Frame++)
+	{
+		const double Time = FMath::Clamp(FrameRate.AsSeconds(Frame), 0., (double)InAnimation->GetPlayLength());
+		const FTransform Transform = UE::Anim::ExtractRootMotionFromAnimationAsset(InAnimation, MirrorTable, 0.0, Time) * ReferenceTransform;
+		const FVector Location = Transform.GetLocation() + Transform.GetUnitAxis(EAxis::Z) * TrackOffset;
+		
+		PDI->DrawTranslucentLine(PrevLocation, Location, NotifyEvent.NotifyColor, SDPG_World, 1.5f, DepthBias, bScreenSpace);
+
+		PrevLocation = Location;
+	}
+
+	const FVector EndLocation = EndRootTransform.GetLocation() + EndRootTransform.GetUnitAxis(EAxis::Z) * TrackOffset;
+	PDI->DrawTranslucentLine(PrevLocation, EndLocation, NotifyEvent.NotifyColor, SDPG_World, 1.5f, DepthBias, bScreenSpace);
+	PrevLocation = EndLocation;
+	
+	// Draw line connecting root motion segment to the warp target.
+	if (FVector::Distance(PrevLocation, WarpPointTransform.GetLocation()) > 5.0)
+	{
+		DrawDashedLine(PDI, PrevLocation, WarpPointTransform.GetLocation(), NotifyEvent.NotifyColor, 5.0f, SDPG_World, DepthBias);
+	}
+
+	// Draw vertical ticks indicating start and end locations.
+	constexpr double RangeTickSize = 5.0;
+	PDI->DrawTranslucentLine(StartRootTransform.GetLocation(), StartRootTransform.GetLocation() + StartRootTransform.GetUnitAxis(EAxis::Z) * RangeTickSize, FColor::Black.WithAlpha(128), SDPG_World, 1, DepthBias, bScreenSpace);
+	PDI->DrawTranslucentLine(EndRootTransform.GetLocation(), EndRootTransform.GetLocation() + EndRootTransform.GetUnitAxis(EAxis::Z) * RangeTickSize, FColor::Black.WithAlpha(128), SDPG_World, 1, DepthBias, bScreenSpace);
+	
+	// Draw warp target transform
+	constexpr double WarpPointSize = 10.0;
+	const FVector WarpLocation = WarpPointTransform.GetLocation();
+	const FVector WarpAxisX = WarpPointTransform.GetUnitAxis(EAxis::X) * WarpPointSize;
+	const FVector WarpAxisY = WarpPointTransform.GetUnitAxis(EAxis::Y) * WarpPointSize;
+	const FVector WarpAxisZ = WarpPointTransform.GetUnitAxis(EAxis::Z) * WarpPointSize;
+	PDI->DrawLine(WarpLocation, WarpLocation + WarpAxisX, FColor::Red, SDPG_Foreground, 1.0f, DepthBias, bScreenSpace);
+	PDI->DrawLine(WarpLocation, WarpLocation + WarpAxisY, FColor::Green, SDPG_Foreground, 1.0f, DepthBias, bScreenSpace);
+	PDI->DrawLine(WarpLocation, WarpLocation + WarpAxisZ, FColor::Blue, SDPG_Foreground, 1.0f, DepthBias, bScreenSpace);
+	
+
+}
+
+void URootMotionModifier_SkewWarp::DrawCanvasInEditor(FCanvas& Canvas, FSceneView& View, USkeletalMeshComponent* MeshComp, const UAnimSequenceBase* InAnimation, const FAnimNotifyEvent& NotifyEvent) const
+{
+	// Early out if the animation does not have root motion.
+	if (!InAnimation || !InAnimation->HasRootMotion())
+	{
+		return;
+	}
+
+	const float NotifyEndTime = NotifyEvent.GetEndTriggerTime();
+
+	const UMirrorDataTable* MirrorTable = nullptr;
+	FTransform ReferenceTransform = FTransform::Identity;
+	if (UDebugSkelMeshComponent* DebugMeshComp = Cast<UDebugSkelMeshComponent>(MeshComp))
+	{
+		ReferenceTransform = DebugMeshComp->RootMotionReferenceTransform;
+		MirrorTable = DebugMeshComp->PreviewInstance ? DebugMeshComp->PreviewInstance->GetMirrorDataTable() : nullptr;
+	}
+
+	const FTransform WarpPointTransform = GetDebugWarpPointTransform(MeshComp, InAnimation, MirrorTable, NotifyEndTime) * ReferenceTransform;
+
+	FVector2D PixelLocation;
+	if (View.WorldToPixel(WarpPointTransform.GetLocation(),PixelLocation))
+	{
+		PixelLocation.X = FMath::RoundToFloat(PixelLocation.X);
+		PixelLocation.Y = FMath::RoundToFloat(PixelLocation.Y);
+
+		constexpr FColor LabelColor(200, 200, 200);
+		constexpr FLinearColor ShadowColor(0, 0, 0, 0.3f);
+		const UFont* SmallFont = GEngine->GetSmallFont();
+
+		Canvas.DrawShadowedString(PixelLocation.X, PixelLocation.Y, *WarpTargetName.ToString(), SmallFont, NotifyEvent.NotifyColor, ShadowColor);
+		PixelLocation.Y += SmallFont->GetMaxCharHeight();
+	}
+}
+#endif
\ No newline at end of file
diff --git a/Engine/Plugins/Animation/MotionWarping/Source/MotionWarping/Public/AnimNotifyState_MotionWarping.h b/Engine/Plugins/Animation/MotionWarping/Source/MotionWarping/Public/AnimNotifyState_MotionWarping.h
--- a/Engine/Plugins/Animation/MotionWarping/Source/MotionWarping/Public/AnimNotifyState_MotionWarping.h
+++ b/Engine/Plugins/Animation/MotionWarping/Source/MotionWarping/Public/AnimNotifyState_MotionWarping.h
@@ -51,6 +51,9 @@ public:
 
 #if WITH_EDITOR
 	virtual void ValidateAssociatedAssets() override;
+
+	virtual void DrawInEditor(class FPrimitiveDrawInterface* PDI, USkeletalMeshComponent* MeshComp, const UAnimSequenceBase* Animation, const FAnimNotifyEvent& NotifyEvent) const override;
+	virtual void DrawCanvasInEditor(FCanvas& Canvas, FSceneView& View, USkeletalMeshComponent* MeshComp, const UAnimSequenceBase* Animation, const FAnimNotifyEvent& NotifyEvent) const override;
 #endif
 };
 
diff --git a/Engine/Plugins/Animation/MotionWarping/Source/MotionWarping/Public/RootMotionModifier.h b/Engine/Plugins/Animation/MotionWarping/Source/MotionWarping/Public/RootMotionModifier.h
--- a/Engine/Plugins/Animation/MotionWarping/Source/MotionWarping/Public/RootMotionModifier.h
+++ b/Engine/Plugins/Animation/MotionWarping/Source/MotionWarping/Public/RootMotionModifier.h
@@ -134,6 +134,11 @@ public:
 
 	FORCEINLINE const UAnimSequenceBase* GetAnimation() const { return Animation.Get(); }
 
+#if WITH_EDITOR	
+	virtual void DrawInEditor(class FPrimitiveDrawInterface* PDI, USkeletalMeshComponent* MeshComp, const UAnimSequenceBase* InAnimation, const FAnimNotifyEvent& NotifyEvent) const {}
+	virtual void DrawCanvasInEditor(class FCanvas& Canvas, class FSceneView& View, USkeletalMeshComponent* MeshComp, const UAnimSequenceBase* InAnimation, const FAnimNotifyEvent& NotifyEvent) const {}
+#endif	
+
 private:
 
 	friend UMotionWarpingComponent;
diff --git a/Engine/Plugins/Animation/MotionWarping/Source/MotionWarping/Public/RootMotionModifier_SkewWarp.h b/Engine/Plugins/Animation/MotionWarping/Source/MotionWarping/Public/RootMotionModifier_SkewWarp.h
--- a/Engine/Plugins/Animation/MotionWarping/Source/MotionWarping/Public/RootMotionModifier_SkewWarp.h
+++ b/Engine/Plugins/Animation/MotionWarping/Source/MotionWarping/Public/RootMotionModifier_SkewWarp.h
@@ -15,9 +15,15 @@ public:
 	URootMotionModifier_SkewWarp(const FObjectInitializer& ObjectInitializer);
 
 	virtual FTransform ProcessRootMotion(const FTransform& InRootMotion, float DeltaSeconds) override;
-
+	
 	static FVector WarpTranslation(const FTransform& CurrentTransform, const FVector& DeltaTranslation, const FVector& TotalTranslation, const FVector& TargetLocation);
 
+#if WITH_EDITOR	
+	virtual void DrawInEditor(FPrimitiveDrawInterface* PDI, USkeletalMeshComponent* MeshComp, const UAnimSequenceBase* Animation, const FAnimNotifyEvent& NotifyEvent) const override;
+	virtual void DrawCanvasInEditor(FCanvas& Canvas, FSceneView& View, USkeletalMeshComponent* MeshComp, const UAnimSequenceBase* Animation, const FAnimNotifyEvent& NotifyEvent) const override;
+	FTransform GetDebugWarpPointTransform(USkeletalMeshComponent* MeshComp, const UAnimSequenceBase* InAnimation, const UMirrorDataTable* MirrorTable, const float NotifyEndTime) const;
+#endif	
+
 	UFUNCTION(BlueprintCallable, Category = "Motion Warping")
 	static URootMotionModifier_SkewWarp* AddRootMotionModifierSkewWarp(
 		UPARAM(DisplayName = "Motion Warping Comp") UMotionWarpingComponent* InMotionWarpingComp,
diff --git a/Engine/Source/Editor/Persona/Private/AnimViewportShowCommands.cpp b/Engine/Source/Editor/Persona/Private/AnimViewportShowCommands.cpp
--- a/Engine/Source/Editor/Persona/Private/AnimViewportShowCommands.cpp
+++ b/Engine/Source/Editor/Persona/Private/AnimViewportShowCommands.cpp
@@ -71,6 +71,9 @@ void FAnimViewportShowCommands::RegisterCommands()
 	UI_COMMAND(ShowAllSections, "Show All Sections", "Display all sections including cloth mapped sections", EUserInterfaceActionType::RadioButton, FInputChord());
 	UI_COMMAND(ShowOnlyClothSections, "Show Only Cloth Sections", "Display only cloth mapped sections", EUserInterfaceActionType::RadioButton, FInputChord());
 	UI_COMMAND(HideOnlyClothSections, "Hide Only Cloth Sections", "Display all except cloth mapped sections", EUserInterfaceActionType::RadioButton, FInputChord());
+
+	UI_COMMAND(ShowNotificationVisualizations, "Show Notification", "Display animation notifications visualizations in viewport", EUserInterfaceActionType::ToggleButton, FInputChord());
+	UI_COMMAND(ShowRootMotionVisualization, "Show Root Motion", "Display root motion visualizations in viewport", EUserInterfaceActionType::ToggleButton, FInputChord());
 }
 
 #undef LOCTEXT_NAMESPACE
\ No newline at end of file
diff --git a/Engine/Source/Editor/Persona/Private/AnimViewportShowCommands.h b/Engine/Source/Editor/Persona/Private/AnimViewportShowCommands.h
--- a/Engine/Source/Editor/Persona/Private/AnimViewportShowCommands.h
+++ b/Engine/Source/Editor/Persona/Private/AnimViewportShowCommands.h
@@ -153,6 +153,12 @@ public:
 
 	TSharedPtr< FUICommandInfo > PauseClothWithAnim;
 
+	/** Show Notification Visualizations */
+	TSharedPtr< FUICommandInfo > ShowNotificationVisualizations;
+
+	/** Show Root motion visualization */
+	TSharedPtr< FUICommandInfo > ShowRootMotionVisualization;
+
 public:
 	/** Registers our commands with the binding system */
 	virtual void RegisterCommands() override;
diff --git a/Engine/Source/Editor/Persona/Private/AnimationEditorViewportClient.cpp b/Engine/Source/Editor/Persona/Private/AnimationEditorViewportClient.cpp
--- a/Engine/Source/Editor/Persona/Private/AnimationEditorViewportClient.cpp
+++ b/Engine/Source/Editor/Persona/Private/AnimationEditorViewportClient.cpp
@@ -47,6 +47,9 @@
 #include "UObject/UE5MainStreamObjectVersion.h"
 #include "Animation/AnimCompositeBase.h"
 #include "AudioEditorSettings.h"
+#include "Animation/AnimNotifies/AnimNotifyState.h"
+#include "Animation/AnimMontage.h"
+#include "Animation/AnimSequenceHelpers.h"
 
 namespace {
 	static const float AnimationEditorViewport_RotateSpeed = 0.02f;
@@ -586,6 +589,11 @@ void FAnimationViewportClient::Draw(const FSceneView* View, FPrimitiveDrawInterf
 			{
 				DrawAttributes(PreviewMeshComponent, PDI);
 			}
+
+			DrawNotifies(PreviewMeshComponent, PDI);
+
+			DrawRootMotionTrajectory(PreviewMeshComponent, PDI);
+			
 		}
 		else if (bValidComponent && !bValidSkeletalMesh)
 		{
@@ -636,6 +644,8 @@ void FAnimationViewportClient::DrawCanvas( FViewport& InViewport, FSceneView& Vi
 			ShowAttributeNames(&Canvas, &View, PreviewMeshComponent);
 		}
 
+		DrawCanvasNotifies(PreviewMeshComponent, Canvas, View);
+
 		if (bDrawUVs)
 		{
 			DrawUVsForMesh(Viewport, &Canvas, 1, PreviewMeshComponent);
@@ -1928,6 +1938,109 @@ void FAnimationViewportClient::DrawAttributes(UDebugSkelMeshComponent* MeshCompo
 	}
 }
 
+void FAnimationViewportClient::DrawNotifies(UDebugSkelMeshComponent* MeshComponent, FPrimitiveDrawInterface* PDI) const
+{
+	if (MeshComponent
+		&& MeshComponent->IsNotificationVisualizationsEnabled()
+		&& MeshComponent->GetSkeletalMeshAsset())
+	{
+		if (const UAnimSequenceBase* AnimSequenceBase = Cast<UAnimSequenceBase>(GetAnimPreviewScene()->GetPreviewAnimationAsset()))
+		{
+			for (const FAnimNotifyEvent& Notify : AnimSequenceBase->Notifies)
+			{
+				if (Notify.Notify)
+				{
+					Notify.Notify->DrawInEditor(PDI, MeshComponent, AnimSequenceBase, Notify);
+				}
+				if (Notify.NotifyStateClass)
+				{
+					Notify.NotifyStateClass->DrawInEditor(PDI, MeshComponent, AnimSequenceBase, Notify);
+				}
+			}
+		}
+	}
+}
+
+void FAnimationViewportClient::DrawCanvasNotifies(UDebugSkelMeshComponent* MeshComponent, FCanvas& Canvas, FSceneView& View) const
+{
+	if (MeshComponent
+		&& MeshComponent->IsNotificationVisualizationsEnabled()
+		&& MeshComponent->GetSkeletalMeshAsset())
+	{
+		if (const UAnimSequenceBase* AnimSequenceBase = Cast<UAnimSequenceBase>(GetAnimPreviewScene()->GetPreviewAnimationAsset()))
+		{
+			for (const FAnimNotifyEvent& Notify : AnimSequenceBase->Notifies)
+			{
+				if (Notify.Notify)
+				{
+					Notify.Notify->DrawCanvasInEditor(Canvas, View, MeshComponent, AnimSequenceBase, Notify);
+				}
+				if (Notify.NotifyStateClass)
+				{
+					Notify.NotifyStateClass->DrawCanvasInEditor(Canvas, View, MeshComponent, AnimSequenceBase, Notify);
+				}
+			}
+		}
+	}
+}
+
+void FAnimationViewportClient::DrawRootMotionTrajectory(UDebugSkelMeshComponent* MeshComponent, FPrimitiveDrawInterface* PDI) const
+{
+	constexpr float DepthBias = 2.0f;
+	constexpr bool bScreenSpace = true;
+
+	if (MeshComponent
+		&& MeshComponent->IsRootMotionVisualizationsEnabled()
+		&& MeshComponent->GetSkeletalMeshAsset()
+		&& MeshComponent->DoesCurrentAssetHaveRootMotion())
+	{
+		const FTransform& ReferenceTransform = MeshComponent->RootMotionReferenceTransform;
+		const UMirrorDataTable* MirrorTable = MeshComponent->PreviewInstance ? MeshComponent->PreviewInstance->GetMirrorDataTable() : nullptr;
+		
+		if (const UAnimSequenceBase* AnimSequenceBase = Cast<UAnimSequenceBase>(GetAnimPreviewScene()->GetPreviewAnimationAsset()))
+		{
+			// Draw root motion trajectory
+			const int32 NumFrames = AnimSequenceBase->GetNumberOfSampledKeys();
+			const FFrameRate FrameRate = AnimSequenceBase->GetSamplingFrameRate();
+
+			const FColor TrajectoryColor = FColor::Black.WithAlpha(64);
+
+			FVector PrevLocation;
+			for (int32 Frame = 0; Frame <= NumFrames; Frame++)
+			{
+				const double Time = FMath::Clamp(FrameRate.AsSeconds(Frame), 0., (double)AnimSequenceBase->GetPlayLength());
+				const FTransform Transform = UE::Anim::ExtractRootMotionFromAnimationAsset(AnimSequenceBase, MirrorTable, 0.0, Time) * ReferenceTransform;
+				const FVector Location = Transform.GetLocation();
+
+				const bool bFirstOrLastPoint = Frame == 0 || Frame == NumFrames;
+				PDI->DrawPoint(Location, TrajectoryColor, bFirstOrLastPoint ? 2.0f : 1.0f, SDPG_World);
+				
+				if (Frame > 0)
+				{
+					PDI->DrawTranslucentLine(PrevLocation, Location, TrajectoryColor, SDPG_World, 1.0f, DepthBias, bScreenSpace);
+				}
+				PrevLocation = Location;
+			}
+
+			// Draw current location on the root motion.
+			{
+				const float CurrentTime = MeshComponent->GetPosition();
+				const FTransform Transform = UE::Anim::ExtractRootMotionFromAnimationAsset(AnimSequenceBase, MirrorTable, 0.0, CurrentTime) * ReferenceTransform;
+
+				constexpr double CurrentPosTickSize = 5.0;
+
+				const FVector Location = Transform.GetLocation();
+				const FVector AxisX = Transform.GetUnitAxis(EAxis::X) * CurrentPosTickSize;
+				const FVector AxisY = Transform.GetUnitAxis(EAxis::Y) * CurrentPosTickSize;
+				const FVector AxisZ = Transform.GetUnitAxis(EAxis::Z) * CurrentPosTickSize;
+				PDI->DrawTranslucentLine(Location, Location + AxisX, FColor::Red.WithAlpha(128), SDPG_World, 1.0f, DepthBias, bScreenSpace);
+				PDI->DrawTranslucentLine(Location, Location + AxisY, FColor::Green.WithAlpha(128), SDPG_World, 1.0f, DepthBias, bScreenSpace);
+				PDI->DrawTranslucentLine(Location, Location + AxisZ, FColor::Blue.WithAlpha(128), SDPG_World, 1.0f, DepthBias, bScreenSpace);
+			}			
+		}
+	}
+}
+
 void FAnimationViewportClient::DrawSockets(const UDebugSkelMeshComponent* InPreviewMeshComponent, TArray<USkeletalMeshSocket*>& InSockets, FSelectedSocketInfo InSelectedSocket, FPrimitiveDrawInterface* PDI, bool bUseSkeletonSocketColor)
 {
 	if (InPreviewMeshComponent && InPreviewMeshComponent->GetSkeletalMeshAsset())
diff --git a/Engine/Source/Editor/Persona/Private/SAnimViewportToolBar.cpp b/Engine/Source/Editor/Persona/Private/SAnimViewportToolBar.cpp
--- a/Engine/Source/Editor/Persona/Private/SAnimViewportToolBar.cpp
+++ b/Engine/Source/Editor/Persona/Private/SAnimViewportToolBar.cpp
@@ -752,6 +752,12 @@ TSharedRef<SWidget> SAnimViewportToolBar::GenerateCharacterMenu() const
 						Section.AddMenuEntry(FAnimViewportShowCommands::Get().ProcessRootMotionLoopAndReset);
 					}
 
+					{
+						FToolMenuSection& Section = InSubMenu->AddSection("AnimViewportVisualization", LOCTEXT("CharacterMenu_VisualizationsLabel", "Visualizations"));
+						Section.AddMenuEntry(FAnimViewportShowCommands::Get().ShowNotificationVisualizations);
+						Section.AddMenuEntry(FAnimViewportShowCommands::Get().ShowRootMotionVisualization);
+					}
+
 					{
 						FToolMenuSection& Section = InSubMenu->AddSection("AnimViewportAnimation", LOCTEXT("CharacterMenu_Actions_AnimationAsset", "Animation"));
 						Section.AddMenuEntry(FAnimViewportShowCommands::Get().ShowRawAnimation);
diff --git a/Engine/Source/Editor/Persona/Private/SAnimationEditorViewport.cpp b/Engine/Source/Editor/Persona/Private/SAnimationEditorViewport.cpp
--- a/Engine/Source/Editor/Persona/Private/SAnimationEditorViewport.cpp
+++ b/Engine/Source/Editor/Persona/Private/SAnimationEditorViewport.cpp
@@ -1103,6 +1103,18 @@ void SAnimationEditorViewportTabBody::BindCommands()
 		FIsActionChecked::CreateSP(this, &SAnimationEditorViewportTabBody::CanUseProcessRootMotionMode, EProcessRootMotionMode::Loop),
 		FIsActionChecked::CreateSP(this, &SAnimationEditorViewportTabBody::IsProcessRootMotionModeSet, EProcessRootMotionMode::Loop));
 
+	CommandList.MapAction(
+		ViewportShowMenuCommands.ShowNotificationVisualizations,
+		FExecuteAction::CreateSP(this, &SAnimationEditorViewportTabBody::ToggleNotificationVisualizations),
+		FIsActionChecked(),
+		FIsActionChecked::CreateSP(this, &SAnimationEditorViewportTabBody::IsNotificationVisualizationsEnabled));
+
+	CommandList.MapAction(
+		ViewportShowMenuCommands.ShowRootMotionVisualization,
+		FExecuteAction::CreateSP(this, &SAnimationEditorViewportTabBody::ToggleRootMotionVisualizations),
+		FIsActionChecked(),
+		FIsActionChecked::CreateSP(this, &SAnimationEditorViewportTabBody::IsRootMotionVisualizationsEnabled));
+
 	CommandList.EndGroup();
 
 	CommandList.MapAction(
@@ -2139,6 +2151,41 @@ bool SAnimationEditorViewportTabBody::CanUseProcessRootMotionMode(EProcessRootMo
 	return false;
 }
 
+void SAnimationEditorViewportTabBody::ToggleNotificationVisualizations()
+{
+	if (UDebugSkelMeshComponent* PreviewComponent = GetPreviewScene()->GetPreviewMeshComponent())
+	{
+		PreviewComponent->SetShowNotificationVisualizations(!PreviewComponent->IsNotificationVisualizationsEnabled());
+	}
+}
+
+bool SAnimationEditorViewportTabBody::IsNotificationVisualizationsEnabled() const
+{
+	if (const UDebugSkelMeshComponent* PreviewComponent = GetPreviewScene()->GetPreviewMeshComponent())
+	{
+		return PreviewComponent->IsNotificationVisualizationsEnabled();
+	}
+	return false;
+}
+
+void SAnimationEditorViewportTabBody::ToggleRootMotionVisualizations()
+{
+	if (UDebugSkelMeshComponent* PreviewComponent = GetPreviewScene()->GetPreviewMeshComponent())
+	{
+		PreviewComponent->SetShowRootMotionVisualizations(!PreviewComponent->IsRootMotionVisualizationsEnabled());
+	}
+}
+
+bool SAnimationEditorViewportTabBody::IsRootMotionVisualizationsEnabled() const
+{
+	if (const UDebugSkelMeshComponent* PreviewComponent = GetPreviewScene()->GetPreviewMeshComponent())
+	{
+		return PreviewComponent->IsRootMotionVisualizationsEnabled();
+	}
+	return false;
+}
+
+
 bool SAnimationEditorViewportTabBody::IsClothSimulationEnabled() const
 {
 	UDebugSkelMeshComponent* PreviewComponent = GetPreviewScene()->GetPreviewMeshComponent();
diff --git a/Engine/Source/Editor/Persona/Private/SAnimationEditorViewport.h b/Engine/Source/Editor/Persona/Private/SAnimationEditorViewport.h
--- a/Engine/Source/Editor/Persona/Private/SAnimationEditorViewport.h
+++ b/Engine/Source/Editor/Persona/Private/SAnimationEditorViewport.h
@@ -519,6 +519,12 @@ private:
 	/** Whether the supplied mode can be used */
 	bool CanUseProcessRootMotionMode(EProcessRootMotionMode Mode) const;
 
+	void ToggleNotificationVisualizations();
+	bool IsNotificationVisualizationsEnabled() const;
+
+	void ToggleRootMotionVisualizations();
+	bool IsRootMotionVisualizationsEnabled() const;
+
 private:
 	/** Selected Turn Table speed  */
 	EAnimationPlaybackSpeeds::Type SelectedTurnTableSpeed;
diff --git a/Engine/Source/Editor/Persona/Public/AnimationEditorViewportClient.h b/Engine/Source/Editor/Persona/Public/AnimationEditorViewportClient.h
--- a/Engine/Source/Editor/Persona/Public/AnimationEditorViewportClient.h
+++ b/Engine/Source/Editor/Persona/Public/AnimationEditorViewportClient.h
@@ -443,6 +443,15 @@ private:
 	/** Draws active transform attributes */
 	void DrawAttributes(UDebugSkelMeshComponent* MeshComponent, FPrimitiveDrawInterface* PDI) const;
 
+	/** Draws visualization from animation notifies into viewport. */
+	void DrawNotifies(UDebugSkelMeshComponent* MeshComponent, FPrimitiveDrawInterface* PDI) const;
+
+	/* Draws visualization from animation notifies into canvas. */
+	void DrawCanvasNotifies(UDebugSkelMeshComponent* MeshComponent, FCanvas& Canvas, FSceneView& View) const;
+
+	/** Draws root motion trajectory */
+	void DrawRootMotionTrajectory(UDebugSkelMeshComponent* MeshComponent, FPrimitiveDrawInterface* PDI) const;
+
 	/** Draws bones from watched poses*/
 	void DrawWatchedPoses(UDebugSkelMeshComponent * MeshComponent, FPrimitiveDrawInterface* PDI);
 
diff --git a/Engine/Source/Editor/UnrealEd/Classes/Animation/DebugSkelMeshComponent.h b/Engine/Source/Editor/UnrealEd/Classes/Animation/DebugSkelMeshComponent.h
--- a/Engine/Source/Editor/UnrealEd/Classes/Animation/DebugSkelMeshComponent.h
+++ b/Engine/Source/Editor/UnrealEd/Classes/Animation/DebugSkelMeshComponent.h
@@ -234,6 +234,14 @@ class UDebugSkelMeshComponent : public USkeletalMeshComponent
 	UPROPERTY(Transient)
 	uint32 bDisplaySourceAnimation:1;
 
+	/** Display Notification visualizations in viewport */
+	UPROPERTY(Transient)
+	uint32 bShowNotificationVisualizations:1;
+
+	/** Display Root motion visualizations in viewport */
+	UPROPERTY(Transient)
+	uint32 bShowRootMotionVisualizations:1;
+
 	/** Display Bound **/
 	UPROPERTY(transient)
 	bool bDisplayBound;
@@ -306,7 +314,7 @@ class UDebugSkelMeshComponent : public USkeletalMeshComponent
 	/** Accumulated root motion. */
 	FTransform RootMotionTransform;
 
-	/** Transform representing the actor transform at the beginning of the animation. */
+	/** Transform representing the actor transform at the beginning of the animation sequence. */
 	FTransform RootMotionReferenceTransform;
 
 	/** Array of bones to render bone weights for */
@@ -491,6 +499,14 @@ class UDebugSkelMeshComponent : public USkeletalMeshComponent
 	/** Whether the current asset or animation blueprint is using root motion */
 	UNREALED_API bool DoesCurrentAssetHaveRootMotion() const;
 
+	/** Sets flags whether we notification visualizations should be drawn in the viewport. */
+	UNREALED_API void SetShowNotificationVisualizations(const bool bShow) { bShowNotificationVisualizations = bShow; }
+	UNREALED_API bool IsNotificationVisualizationsEnabled() const { return bShowNotificationVisualizations; }
+
+	/** Sets flags whether we root motion visualization should be drawn in the viewport. */
+	UNREALED_API void SetShowRootMotionVisualizations(const bool bShow) { bShowRootMotionVisualizations = bShow; }
+	UNREALED_API bool IsRootMotionVisualizationsEnabled() const { return bShowRootMotionVisualizations; }
+
 	/** Whether the current LOD of the debug mesh is being synced with the attached (preview) mesh instance. */
 	UNREALED_API bool IsTrackingAttachedLOD() const;
 
diff --git a/Engine/Source/Editor/UnrealEd/Private/Animation/DebugSkelMeshComponent.cpp b/Engine/Source/Editor/UnrealEd/Private/Animation/DebugSkelMeshComponent.cpp
--- a/Engine/Source/Editor/UnrealEd/Private/Animation/DebugSkelMeshComponent.cpp
+++ b/Engine/Source/Editor/UnrealEd/Private/Animation/DebugSkelMeshComponent.cpp
@@ -40,6 +40,9 @@ UDebugSkelMeshComponent::UDebugSkelMeshComponent(const FObjectInitializer& Objec
 	bMeshSocketsVisible = true;
 	bSkeletonSocketsVisible = true;
 
+	bShowNotificationVisualizations = false;
+	bShowRootMotionVisualizations = false;
+
 	TurnTableSpeedScaling = 1.f;
 	TurnTableMode = EPersonaTurnTableMode::Stopped;
 	TurntableTransform = FTransform::Identity;
@@ -284,13 +287,14 @@ void UDebugSkelMeshComponent::ConsumeRootMotion(const FVector& FloorMin, const F
 			// Loop and Reset Mode: Preview mesh will consume root motion resetting the position back to the origin every time the animation loops
 			else if (ProcessRootMotionMode == EProcessRootMotionMode::LoopAndReset)
 			{
-				RootMotionReferenceTransform = FTransform::Identity;
-				
 				if (bLooped)
 				{
 					const FTransform InitialTransform = UE::Anim::ExtractRootTransformFromAnimationAsset(PreviewInstance->CurrentAsset, SectionStartPosition);
 					const FTransform RootMotionDelta = UE::Anim::ExtractRootMotionFromAnimationAsset(PreviewInstance->CurrentAsset, PreviewInstance->GetMirrorDataTable(), SectionStartPosition, CurrentTime);
-					RootMotionTransform = InitialTransform * RootMotionDelta;
+					RootMotionTransform = RootMotionDelta * InitialTransform;
+
+					// Reference transform is always relative to the beginning of the animation sequence. 
+					RootMotionReferenceTransform = UE::Anim::ExtractRootTransformFromAnimationAsset(PreviewInstance->CurrentAsset, 0.0f);
 				}
 				else
 				{
@@ -431,13 +435,17 @@ void UDebugSkelMeshComponent::SetProcessRootMotionModeInternal(EProcessRootMotio
 	
 		const FTransform InitialTransform = UE::Anim::ExtractRootTransformFromAnimationAsset(PreviewInstance->CurrentAsset, SectionStartPosition);
 		const FTransform RootMotionDelta = UE::Anim::ExtractRootMotionFromAnimationAsset(PreviewInstance->CurrentAsset, PreviewInstance->GetMirrorDataTable(), SectionStartPosition, CurrentTime);
-		RootMotionTransform = InitialTransform * RootMotionDelta;
+		RootMotionTransform = RootMotionDelta * InitialTransform;
+		
+		// Reference transform is always relative to the beginning of the animation sequence. 
+		RootMotionReferenceTransform = UE::Anim::ExtractRootTransformFromAnimationAsset(PreviewInstance->CurrentAsset, 0.0f);
 	}
 	else if (ProcessRootMotionMode == EProcessRootMotionMode::Ignore)
 	{
 		RootMotionTransform = FTransform::Identity;
+		RootMotionReferenceTransform = FTransform::Identity;
 	}
-	
+
 	SetRelativeTransform(RootMotionTransform * TurntableTransform, /*bSweep*/false, /*OutSweepResult*/nullptr, ETeleportType::ResetPhysics);
 }
 
@@ -696,8 +704,11 @@ void UDebugSkelMeshComponent::OnMirrorDataTableChanged()
 	
 		const FTransform InitialTransform = UE::Anim::ExtractRootTransformFromAnimationAsset(PreviewInstance->CurrentAsset, SectionStartPosition);
 		const FTransform RootMotionDelta = UE::Anim::ExtractRootMotionFromAnimationAsset(PreviewInstance->CurrentAsset, PreviewInstance->GetMirrorDataTable(), SectionStartPosition, CurrentTime);
-		RootMotionTransform = InitialTransform * RootMotionDelta;
-	
+		RootMotionTransform = RootMotionDelta * InitialTransform;
+
+		// Reference transform is always relative to the beginning of the animation sequence. 
+		RootMotionReferenceTransform = UE::Anim::ExtractRootTransformFromAnimationAsset(PreviewInstance->CurrentAsset, 0.0f);;
+
 		SetRelativeTransform(RootMotionTransform * TurntableTransform, /*bSweep*/false, /*OutSweepResult*/nullptr, ETeleportType::ResetPhysics);
 	}
 }
diff --git a/Engine/Source/Runtime/Engine/Classes/Animation/AnimNotifies/AnimNotify.h b/Engine/Source/Runtime/Engine/Classes/Animation/AnimNotifies/AnimNotify.h
--- a/Engine/Source/Runtime/Engine/Classes/Animation/AnimNotifies/AnimNotify.h
+++ b/Engine/Source/Runtime/Engine/Classes/Animation/AnimNotifies/AnimNotify.h
@@ -13,6 +13,12 @@
 class UAnimSequenceBase;
 class USkeletalMeshComponent;
 struct FAnimNotifyEvent;
+#if WITH_EDITORONLY_DATA
+class FPrimitiveDrawInterface;
+class FCanvas;
+class FSceneView;
+#endif
+
 USTRUCT()
 struct FBranchingPointNotifyPayload
 {
@@ -69,6 +75,8 @@ class UAnimNotify : public UObject
 	virtual void OnAnimNotifyCreatedInEditor(FAnimNotifyEvent& ContainingAnimNotifyEvent) {};
 	virtual bool CanBePlaced(UAnimSequenceBase* Animation) const { return true; }
 	virtual void ValidateAssociatedAssets() {}
+	virtual void DrawInEditor(FPrimitiveDrawInterface* PDI, USkeletalMeshComponent* MeshComp, const UAnimSequenceBase* Animation, const FAnimNotifyEvent& NotifyEvent) const {}
+	virtual void DrawCanvasInEditor(FCanvas& Canvas, FSceneView& View, USkeletalMeshComponent* MeshComp, const UAnimSequenceBase* Animation, const FAnimNotifyEvent& NotifyEvent) const {}
 
 	/** Override this to prevent firing this notify type in animation editors */
 	virtual bool ShouldFireInEditor() { return bShouldFireInEditor; }
diff --git a/Engine/Source/Runtime/Engine/Classes/Animation/AnimNotifies/AnimNotifyState.h b/Engine/Source/Runtime/Engine/Classes/Animation/AnimNotifies/AnimNotifyState.h
--- a/Engine/Source/Runtime/Engine/Classes/Animation/AnimNotifies/AnimNotifyState.h
+++ b/Engine/Source/Runtime/Engine/Classes/Animation/AnimNotifies/AnimNotifyState.h
@@ -14,6 +14,11 @@ class UAnimSequenceBase;
 class USkeletalMeshComponent;
 struct FAnimNotifyEvent;
 struct FBranchingPointNotifyPayload;
+#if WITH_EDITORONLY_DATA
+class FPrimitiveDrawInterface;
+class FCanvas;
+class FSceneView;
+#endif
 
 UCLASS(abstract, editinlinenew, Blueprintable, const, hidecategories=Object, collapsecategories, meta=(ShowWorldContextPin), MinimalAPI)
 class UAnimNotifyState : public UObject
@@ -49,6 +54,8 @@ class UAnimNotifyState : public UObject
 	virtual void OnAnimNotifyCreatedInEditor(FAnimNotifyEvent& ContainingAnimNotifyEvent) {};
 	virtual bool CanBePlaced(UAnimSequenceBase* Animation) const { return true; }
 	virtual void ValidateAssociatedAssets() {}
+	virtual void DrawInEditor(FPrimitiveDrawInterface* PDI, USkeletalMeshComponent* MeshComp, const UAnimSequenceBase* Animation, const FAnimNotifyEvent& NotifyEvent) const {}
+	virtual void DrawCanvasInEditor(FCanvas& Canvas, FSceneView& View, USkeletalMeshComponent* MeshComp, const UAnimSequenceBase* Animation, const FAnimNotifyEvent& NotifyEvent) const {}
 
 	/** Override this to prevent firing this notify state type in animation editors */
 	virtual bool ShouldFireInEditor() { return bShouldFireInEditor; }
