Interaction Proxy: Added anim metada to show Interaction Proxy, and derived skew warp that uses slot as ref transform.
- Added UAnimMetaData_InteractionProxy which allows to associate IP with an animation, includes preview
- Added URootMotionModifier_InteractionSkewWarp which works with the metadata to allow to set a warp target to specific slot
- Added labels for Slots to make it easier to match them with selection
- Changed Seat pivot to be at the center of the seat

diff --git a/Engine/Plugins/Animation/MotionWarping/Source/MotionWarping/Private/RootMotionModifier_SkewWarp.cpp b/Engine/Plugins/Animation/MotionWarping/Source/MotionWarping/Private/RootMotionModifier_SkewWarp.cpp
--- a/Engine/Plugins/Animation/MotionWarping/Source/MotionWarping/Private/RootMotionModifier_SkewWarp.cpp
+++ b/Engine/Plugins/Animation/MotionWarping/Source/MotionWarping/Private/RootMotionModifier_SkewWarp.cpp
@@ -263,7 +263,11 @@ FTransform URootMotionModifier_SkewWarp::GetDebugWarpPointTransform(USkeletalMes
 	}
 	else if (WarpPointAnimProvider == EWarpPointAnimProvider::Static)
 	{
-		WarpPointTransform = WarpPointAnimTransform;
+		// This method returns the warp transform relative to the actor location on the first frame of the animation.
+		// The WarpPointAnimTransform is defined in the same coordinate space as the root motion track.
+		// Adjust the return warp point to be relative to the first frame's Actor transform.
+		const FTransform FirstFrameTransform = UE::Anim::ExtractRootTransformFromAnimationAsset(InAnimation, 0.0f);
+		WarpPointTransform = WarpPointAnimTransform * FirstFrameTransform.Inverse();
 	}
 	else if (WarpPointAnimProvider == EWarpPointAnimProvider::Bone)
 	{
@@ -386,18 +390,21 @@ void URootMotionModifier_SkewWarp::DrawCanvasInEditor(FCanvas& Canvas, FSceneVie
 
 	const FTransform WarpPointTransform = GetDebugWarpPointTransform(MeshComp, InAnimation, MirrorTable, NotifyEndTime) * ReferenceTransform;
 
-	FVector2D PixelLocation;
-	if (View.WorldToPixel(WarpPointTransform.GetLocation(),PixelLocation))
+	if (!WarpTargetName.IsNone())
 	{
-		PixelLocation.X = FMath::RoundToFloat(PixelLocation.X);
-		PixelLocation.Y = FMath::RoundToFloat(PixelLocation.Y);
+		FVector2D PixelLocation;
+		if (View.WorldToPixel(WarpPointTransform.GetLocation(),PixelLocation))
+		{
+			PixelLocation.X = FMath::RoundToFloat(PixelLocation.X);
+			PixelLocation.Y = FMath::RoundToFloat(PixelLocation.Y);
 
-		constexpr FColor LabelColor(200, 200, 200);
-		constexpr FLinearColor ShadowColor(0, 0, 0, 0.3f);
-		const UFont* SmallFont = GEngine->GetSmallFont();
+			constexpr FColor LabelColor(200, 200, 200);
+			constexpr FLinearColor ShadowColor(0, 0, 0, 0.3f);
+			const UFont* SmallFont = GEngine->GetSmallFont();
 
-		Canvas.DrawShadowedString(PixelLocation.X, PixelLocation.Y, *WarpTargetName.ToString(), SmallFont, NotifyEvent.NotifyColor, ShadowColor);
-		PixelLocation.Y += SmallFont->GetMaxCharHeight();
+			Canvas.DrawShadowedString(PixelLocation.X, PixelLocation.Y, *WarpTargetName.ToString(), SmallFont, NotifyEvent.NotifyColor, ShadowColor);
+			PixelLocation.Y += SmallFont->GetMaxCharHeight();
+		}
 	}
 }
 #endif
\ No newline at end of file
