Motion Warping - Fixed bug that would cause motion warping to inject and warp root motion during the same warping window if the window overlaps an area in the animation with and without root motion

diff --git a/Engine/Plugins/Animation/MotionWarping/Source/MotionWarping/Private/RootMotionModifier.cpp b/Engine/Plugins/Animation/MotionWarping/Source/MotionWarping/Private/RootMotionModifier.cpp
--- a/Engine/Plugins/Animation/MotionWarping/Source/MotionWarping/Private/RootMotionModifier.cpp
+++ b/Engine/Plugins/Animation/MotionWarping/Source/MotionWarping/Private/RootMotionModifier.cpp
@@ -194,6 +194,8 @@ void URootMotionModifier::OnStateChanged(ERootMotionModifierState LastState)
 			const FVector CurrentLocation = (CharacterOwner->GetActorLocation() - CurrentRotation.GetUpVector() * CapsuleHalfHeight);
 			StartTransform = FTransform(CurrentRotation, CurrentLocation);
 
+			TotalRootMotionWithinWindow = UMotionWarpingUtilities::ExtractRootMotionFromAnimation(Animation.Get(), StartTime, EndTime);
+
 			OnActivateDelegate.ExecuteIfBound(OwnerComp, this);
 		}
 		else if (LastState == ERootMotionModifierState::Active && (State == ERootMotionModifierState::Disabled || State == ERootMotionModifierState::MarkedForRemoval))
diff --git a/Engine/Plugins/Animation/MotionWarping/Source/MotionWarping/Private/RootMotionModifier_SkewWarp.cpp b/Engine/Plugins/Animation/MotionWarping/Source/MotionWarping/Private/RootMotionModifier_SkewWarp.cpp
--- a/Engine/Plugins/Animation/MotionWarping/Source/MotionWarping/Private/RootMotionModifier_SkewWarp.cpp
+++ b/Engine/Plugins/Animation/MotionWarping/Source/MotionWarping/Private/RootMotionModifier_SkewWarp.cpp
@@ -164,7 +164,7 @@ FTransform URootMotionModifier_SkewWarp::ProcessRootMotion(const FTransform& InR
 		}
 
 		// if there is translation in the animation, warp it
-		if (!TotalTranslation.IsNearlyZero())
+		if (!TotalRootMotionWithinWindow.GetTranslation().IsNearlyZero())
 		{
 			if (!DeltaTranslation.IsNearlyZero())
 			{
diff --git a/Engine/Plugins/Animation/MotionWarping/Source/MotionWarping/Public/RootMotionModifier.h b/Engine/Plugins/Animation/MotionWarping/Source/MotionWarping/Public/RootMotionModifier.h
--- a/Engine/Plugins/Animation/MotionWarping/Source/MotionWarping/Public/RootMotionModifier.h
+++ b/Engine/Plugins/Animation/MotionWarping/Source/MotionWarping/Public/RootMotionModifier.h
@@ -100,6 +100,10 @@ public:
 	UPROPERTY(Transient, BlueprintReadOnly, Category = "Defaults")
 	float ActualStartTime = 0.f;
 
+	/** Total root motion within the warping window */
+	UPROPERTY(Transient, BlueprintReadOnly, Category = "Defaults")
+	FTransform TotalRootMotionWithinWindow;
+
 	/** Delegate called when this modifier is activated (starts affecting the root motion) */
 	UPROPERTY()
 	FOnRootMotionModifierDelegate OnActivateDelegate;
