Expands Motion Warping plugin to support non-Character actors using adapters.

Added warping support for MoverComponent-using actors, and fixed

#virtualized

diff --git a/Engine/Plugins/Animation/MotionWarping/Source/MotionWarping/Private/MotionWarpingCharacterAdapter.cpp b/Engine/Plugins/Animation/MotionWarping/Source/MotionWarping/Private/MotionWarpingCharacterAdapter.cpp
new file mode 100644
--- /dev/null
+++ b/Engine/Plugins/Animation/MotionWarping/Source/MotionWarping/Private/MotionWarpingCharacterAdapter.cpp
@@ -0,0 +1,99 @@
+// Copyright Epic Games, Inc. All Rights Reserved.
+
+#include "MotionWarpingCharacterAdapter.h"
+#include "GameFramework/Character.h"
+#include "GameFramework/CharacterMovementComponent.h"
+#include "Components/CapsuleComponent.h"
+
+
+UMotionWarpingCharacterAdapter::UMotionWarpingCharacterAdapter()
+{
+}
+
+UMotionWarpingCharacterAdapter::~UMotionWarpingCharacterAdapter()
+{
+	if (TargetCharacter)
+	{
+		TargetCharacter->GetCharacterMovement()->ProcessRootMotionPreConvertToWorld.Unbind();
+	}
+}
+
+void UMotionWarpingCharacterAdapter::SetCharacter(ACharacter* InCharacter)
+{
+	TargetCharacter = InCharacter;
+	TargetCharacter->GetCharacterMovement()->ProcessRootMotionPreConvertToWorld.BindUObject(this, &UMotionWarpingCharacterAdapter::WarpLocalRootMotionOnCharacter);
+}
+
+AActor* UMotionWarpingCharacterAdapter::GetActor() const
+{ 
+	return Cast<AActor>(TargetCharacter);
+}
+
+USkeletalMeshComponent* UMotionWarpingCharacterAdapter::GetMesh() const
+{ 
+	return TargetCharacter->GetMesh();
+}
+
+FVector UMotionWarpingCharacterAdapter::GetVisualRootLocation() const
+{
+	const float CapsuleHalfHeight = TargetCharacter->GetCapsuleComponent()->GetScaledCapsuleHalfHeight();
+	const FQuat CurrentRotation = TargetCharacter->GetActorQuat();
+	return  (TargetCharacter->GetActorLocation() - CurrentRotation.GetUpVector() * CapsuleHalfHeight);
+}
+
+FVector UMotionWarpingCharacterAdapter::GetBaseVisualTranslationOffset() const
+{
+	return TargetCharacter->GetBaseTranslationOffset();
+}
+
+FQuat UMotionWarpingCharacterAdapter::GetBaseVisualRotationOffset() const
+{ 
+	return TargetCharacter->GetBaseRotationOffset();
+}
+
+FTransform UMotionWarpingCharacterAdapter::WarpLocalRootMotionOnCharacter(const FTransform& LocalRootMotionTransform, UCharacterMovementComponent* TargetMoveComp, float DeltaSeconds)
+{
+	if (WarpLocalRootMotionDelegate.IsBound() && TargetCharacter)
+	{
+		FMotionWarpingUpdateContext WarpingContext;
+		
+		WarpingContext.DeltaSeconds = DeltaSeconds;
+
+		// When replaying saved moves we need to look at the contributor to root motion back then.
+		if (TargetCharacter->bClientUpdating)
+		{
+			const UCharacterMovementComponent* MoveComp = TargetCharacter->GetCharacterMovement();
+			check(MoveComp);
+
+			const FSavedMove_Character* SavedMove = MoveComp->GetCurrentReplayedSavedMove();
+			check(SavedMove);
+
+			if (SavedMove->RootMotionMontage.IsValid())
+			{
+				WarpingContext.Animation = SavedMove->RootMotionMontage.Get();
+				WarpingContext.CurrentPosition = SavedMove->RootMotionTrackPosition;
+				WarpingContext.PreviousPosition = SavedMove->RootMotionPreviousTrackPosition;
+				WarpingContext.PlayRate = SavedMove->RootMotionPlayRateWithScale;
+			}
+		}
+		else // If we are not replaying a move, just use the current root motion montage
+		{
+			if (const FAnimMontageInstance* RootMotionMontageInstance = TargetCharacter->GetRootMotionAnimMontageInstance())
+			{
+				const UAnimMontage* Montage = RootMotionMontageInstance->Montage;
+				check(Montage);
+
+				WarpingContext.Animation = Montage;
+				WarpingContext.CurrentPosition = RootMotionMontageInstance->GetPosition();
+				WarpingContext.PreviousPosition = RootMotionMontageInstance->GetPreviousPosition();
+				WarpingContext.Weight = RootMotionMontageInstance->GetWeight();
+				WarpingContext.PlayRate = RootMotionMontageInstance->Montage->RateScale * RootMotionMontageInstance->GetPlayRate();
+			}
+		}
+
+		// TODO: Consider simply having a pointer to the MWComponent whereby we can call a function on it, rather than using this delegate approach
+		return WarpLocalRootMotionDelegate.Execute(LocalRootMotionTransform, DeltaSeconds, &WarpingContext);
+	}
+
+	return LocalRootMotionTransform;
+}
\ No newline at end of file
diff --git a/Engine/Plugins/Animation/MotionWarping/Source/MotionWarping/Private/MotionWarpingComponent.cpp b/Engine/Plugins/Animation/MotionWarping/Source/MotionWarping/Private/MotionWarpingComponent.cpp
--- a/Engine/Plugins/Animation/MotionWarping/Source/MotionWarping/Private/MotionWarpingComponent.cpp
+++ b/Engine/Plugins/Animation/MotionWarping/Source/MotionWarping/Private/MotionWarpingComponent.cpp
@@ -13,6 +13,7 @@
 #include "GameFramework/Character.h"
 #include "GameFramework/CharacterMovementComponent.h"
 #include "AnimNotifyState_MotionWarping.h"
+#include "MotionWarpingCharacterAdapter.h"
 #include "Net/Core/PushModel/PushModel.h"
 #include "Net/UnrealNetwork.h"
 
@@ -163,6 +164,7 @@ void UMotionWarpingUtilities::GetMotionWarpingWindowsForWarpTargetFromAnimation(
 	}
 }
 
+
 FTransform UMotionWarpingUtilities::CalculateRootTransformRelativeToWarpPointAtTime(const ACharacter& Character, const UAnimSequenceBase* Animation, float Time, const FName& WarpPointBoneName)
 {
 	if (const USkeletalMeshComponent* Mesh = Character.GetMesh())
@@ -194,6 +196,39 @@ FTransform UMotionWarpingUtilities::CalculateRootTransformRelativeToWarpPointAtT
 	return FTransform::Identity;
 }
 
+
+FTransform UMotionWarpingUtilities::CalculateRootTransformRelativeToWarpPointAtTime(const UMotionWarpingBaseAdapter& WarpingAdapter, const UAnimSequenceBase* Animation, float Time, const FName& WarpPointBoneName)
+{
+	if (const USkeletalMeshComponent* Mesh = WarpingAdapter.GetMesh())
+	{
+		if (const UAnimInstance* AnimInstance = Mesh->GetAnimInstance())
+		{
+			const FBoneContainer& FullBoneContainer = AnimInstance->GetRequiredBones();
+			const int32 BoneIndex = FullBoneContainer.GetPoseBoneIndexForBoneName(WarpPointBoneName);
+			if (BoneIndex != INDEX_NONE)
+			{
+				TArray<FBoneIndexType> RequiredBoneIndexArray = { 0, (FBoneIndexType)BoneIndex };
+				FullBoneContainer.GetReferenceSkeleton().EnsureParentsExistAndSort(RequiredBoneIndexArray);
+
+				FBoneContainer LimitedBoneContainer(RequiredBoneIndexArray, UE::Anim::FCurveFilterSettings(UE::Anim::ECurveFilterMode::DisallowAll), *FullBoneContainer.GetAsset());
+
+				FCSPose<FCompactPose> Pose;
+				UMotionWarpingUtilities::ExtractComponentSpacePose(Animation, LimitedBoneContainer, Time, false, Pose);
+
+				// Inverse of mesh's relative rotation. Used to convert root and warp point in the animation from Y forward to X forward
+				const FTransform MeshCompRelativeRotInverse = FTransform(WarpingAdapter.GetBaseVisualRotationOffset().Inverse());
+
+				const FTransform RootTransform = MeshCompRelativeRotInverse * Pose.GetComponentSpaceTransform(FCompactPoseBoneIndex(0));
+				const FTransform WarpPointTransform = MeshCompRelativeRotInverse * Pose.GetComponentSpaceTransform(FCompactPoseBoneIndex(1));
+				return RootTransform.GetRelativeTransform(WarpPointTransform);
+			}
+		}
+	}
+
+	return FTransform::Identity;
+}
+
+
 FTransform UMotionWarpingUtilities::CalculateRootTransformRelativeToWarpPointAtTime(const ACharacter& Character, const UAnimSequenceBase* Animation, float Time, const FTransform& WarpPointTransform)
 {
 	// Inverse of mesh's relative rotation. Used to convert root and warp point in the animation from Y forward to X forward
@@ -202,6 +237,16 @@ FTransform UMotionWarpingUtilities::CalculateRootTransformRelativeToWarpPointAtT
 	return RootTransform.GetRelativeTransform((MeshCompRelativeRotInverse * WarpPointTransform));
 }
 
+
+FTransform UMotionWarpingUtilities::CalculateRootTransformRelativeToWarpPointAtTime(const UMotionWarpingBaseAdapter& WarpingAdapter, const UAnimSequenceBase* Animation, float Time, const FTransform& WarpPointTransform)
+{
+	// Inverse of mesh's relative rotation. Used to convert root and warp point in the animation from Y forward to X forward
+	const FTransform MeshCompRelativeRotInverse = FTransform(WarpingAdapter.GetBaseVisualRotationOffset().Inverse());
+	const FTransform RootTransform = MeshCompRelativeRotInverse * UMotionWarpingUtilities::ExtractRootTransformFromAnimation(Animation, Time);
+	return RootTransform.GetRelativeTransform((MeshCompRelativeRotInverse * WarpPointTransform));
+}
+
+
 // UMotionWarpingComponent
 ///////////////////////////////////////////////////////////////////////
 
@@ -226,13 +271,34 @@ void UMotionWarpingComponent::InitializeComponent()
 {
 	Super::InitializeComponent();
 
-	CharacterOwner = Cast<ACharacter>(GetOwner());
+	// Implicitly support Characters if no other adapter has already been setup
+	if (GetOwnerAdapter() == nullptr)
+	{
+		if (ACharacter* CharacterOwner = Cast<ACharacter>(GetOwner()))
+		{
+			UMotionWarpingCharacterAdapter* CharacterAdapter = CreateOwnerAdapter<UMotionWarpingCharacterAdapter>();
+			CharacterAdapter->SetCharacter(CharacterOwner);
+		}
+	}
+}
+
+UMotionWarpingBaseAdapter* UMotionWarpingComponent::CreateOwnerAdapter(TSubclassOf<UMotionWarpingBaseAdapter> AdapterClass)
+{
+	check(AdapterClass);	
+	OwnerAdapter = NewObject<UMotionWarpingBaseAdapter>(this, AdapterClass);
+	OwnerAdapter->WarpLocalRootMotionDelegate.BindUObject(this, &UMotionWarpingComponent::ProcessRootMotionPreConvertToWorld);
+
+	return OwnerAdapter;
+}
 
-	UCharacterMovementComponent* CharacterMovementComp = CharacterOwner.IsValid() ? CharacterOwner->GetCharacterMovement() : nullptr;
-	if (CharacterMovementComp)
+ACharacter* UMotionWarpingComponent::GetCharacterOwner() const
+{ 
+	if (OwnerAdapter)
 	{
- 		CharacterMovementComp->ProcessRootMotionPreConvertToWorld.BindUObject(this, &UMotionWarpingComponent::ProcessRootMotionPreConvertToWorld);
+		return Cast<ACharacter>(OwnerAdapter->GetActor());
 	}
+
+	return nullptr; 
 }
 
 bool UMotionWarpingComponent::ContainsModifier(const UAnimSequenceBase* Animation, float StartTime, float EndTime) const
@@ -248,8 +314,8 @@ int32 UMotionWarpingComponent::AddModifier(URootMotionModifier* Modifier)
 	if (ensureAlways(Modifier))
 	{
 		UE_LOG(LogMotionWarping, Verbose, TEXT("MotionWarping: RootMotionModifier added. NetMode: %d WorldTime: %f Char: %s Animation: %s [%f %f] [%f %f] Loc: %s Rot: %s"),
-			GetWorld()->GetNetMode(), GetWorld()->GetTimeSeconds(), *GetNameSafe(GetCharacterOwner()), *GetNameSafe(Modifier->Animation.Get()), Modifier->StartTime, Modifier->EndTime, Modifier->PreviousPosition, Modifier->CurrentPosition,
-			*GetCharacterOwner()->GetActorLocation().ToString(), *GetCharacterOwner()->GetActorRotation().ToCompactString());
+			GetWorld()->GetNetMode(), GetWorld()->GetTimeSeconds(), *GetNameSafe(GetOwner()), *GetNameSafe(Modifier->Animation.Get()), Modifier->StartTime, Modifier->EndTime, Modifier->PreviousPosition, Modifier->CurrentPosition,
+			*GetOwner()->GetActorLocation().ToString(), *GetOwner()->GetActorRotation().ToCompactString());
 
 		return Modifiers.Add(Modifier);
 	}
@@ -268,46 +334,8 @@ void UMotionWarpingComponent::DisableAllRootMotionModifiers()
 	}
 }
 
-void UMotionWarpingComponent::Update(float DeltaSeconds)
+void UMotionWarpingComponent::UpdateWithContext(const FMotionWarpingUpdateContext& Context, float DeltaSeconds)
 {
-	const ACharacter* Character = GetCharacterOwner();
-	check(Character);
-
-	FMotionWarpingUpdateContext Context;
-	Context.DeltaSeconds = DeltaSeconds;
-
-	// When replaying saved moves we need to look at the contributor to root motion back then.
-	if (Character->bClientUpdating)
-	{
-		const UCharacterMovementComponent* MoveComp = Character->GetCharacterMovement();
-		check(MoveComp);
-
-		const FSavedMove_Character* SavedMove = MoveComp->GetCurrentReplayedSavedMove();
-		check(SavedMove);
-
-		if(SavedMove->RootMotionMontage.IsValid())
-		{
-			Context.Animation = SavedMove->RootMotionMontage.Get();
-			Context.CurrentPosition = SavedMove->RootMotionTrackPosition;
-			Context.PreviousPosition = SavedMove->RootMotionPreviousTrackPosition;
-			Context.PlayRate = SavedMove->RootMotionPlayRateWithScale;
-		}
-	}
-	else // If we are not replaying a move, just use the current root motion montage
-	{
-		if(const FAnimMontageInstance* RootMotionMontageInstance = Character->GetRootMotionAnimMontageInstance())
-		{
-			const UAnimMontage* Montage = RootMotionMontageInstance->Montage;
-			check(Montage);
-
-			Context.Animation = Montage;
-			Context.CurrentPosition = RootMotionMontageInstance->GetPosition();
-			Context.PreviousPosition = RootMotionMontageInstance->GetPreviousPosition();
-			Context.Weight = RootMotionMontageInstance->GetWeight();
-			Context.PlayRate = RootMotionMontageInstance->Montage->RateScale * RootMotionMontageInstance->GetPlayRate();
-		}
-	}
-
 	if (Context.Animation.IsValid())
 	{
 		const UAnimSequenceBase* Animation = Context.Animation.Get();
@@ -402,8 +430,8 @@ void UMotionWarpingComponent::Update(float DeltaSeconds)
 			if (Modifier->GetState() == ERootMotionModifierState::MarkedForRemoval)
 			{
 				UE_LOG(LogMotionWarping, Verbose, TEXT("MotionWarping: RootMotionModifier removed. NetMode: %d WorldTime: %f Char: %s Animation: %s [%f %f] [%f %f] Loc: %s Rot: %s"),
-					GetWorld()->GetNetMode(), GetWorld()->GetTimeSeconds(), *GetNameSafe(GetCharacterOwner()), *GetNameSafe(Modifier->Animation.Get()), Modifier->StartTime, Modifier->EndTime, Modifier->PreviousPosition, Modifier->CurrentPosition,
-					*GetCharacterOwner()->GetActorLocation().ToString(), *GetCharacterOwner()->GetActorRotation().ToCompactString());
+					GetWorld()->GetNetMode(), GetWorld()->GetTimeSeconds(), *GetNameSafe(GetOwner()), *GetNameSafe(Modifier->Animation.Get()), Modifier->StartTime, Modifier->EndTime, Modifier->PreviousPosition, Modifier->CurrentPosition,
+					*GetOwner()->GetActorLocation().ToString(), *GetOwner()->GetActorRotation().ToCompactString());
 
 				return true;
 			}
@@ -413,7 +441,7 @@ void UMotionWarpingComponent::Update(float DeltaSeconds)
 	}
 }
 
-FTransform UMotionWarpingComponent::ProcessRootMotionPreConvertToWorld(const FTransform& InRootMotion, UCharacterMovementComponent* CharacterMovementComponent, float DeltaSeconds)
+FTransform UMotionWarpingComponent::ProcessRootMotionPreConvertToWorld(const FTransform& InRootMotion, float DeltaSeconds, const FMotionWarpingUpdateContext* InContext)
 {
 #if !(UE_BUILD_SHIPPING || UE_BUILD_TEST)
 	if (FMotionWarpingCVars::CVarMotionWarpingDisable.GetValueOnGameThread() > 0)
@@ -421,9 +449,13 @@ FTransform UMotionWarpingComponent::ProcessRootMotionPreConvertToWorld(const FTr
 		return InRootMotion;
 	}
 #endif
+	if (!InContext)
+	{
+		return InRootMotion;
+	}
 
 	// Check for warping windows and update modifier states
-	Update(DeltaSeconds);
+	UpdateWithContext(*InContext, DeltaSeconds);
 
 	FTransform FinalRootMotion = InRootMotion;
 
@@ -438,11 +470,11 @@ FTransform UMotionWarpingComponent::ProcessRootMotionPreConvertToWorld(const FTr
 
 #if !(UE_BUILD_SHIPPING || UE_BUILD_TEST)
 	const int32 DebugLevel = FMotionWarpingCVars::CVarMotionWarpingDebug.GetValueOnGameThread();
-	if (DebugLevel >= 2)
+	if (DebugLevel >= 2 && OwnerAdapter)
 	{
 		const float DrawDebugDuration = FMotionWarpingCVars::CVarMotionWarpingDrawDebugDuration.GetValueOnGameThread();
 		const float PointSize = 7.f;
-		const FVector ActorFeetLocation = CharacterMovementComponent->GetActorFeetLocation();
+		const FVector ActorFeetLocation = OwnerAdapter->GetVisualRootLocation();
 		if (Modifiers.Num() > 0)
 		{
 			if (!OriginalRootMotionAccum.IsSet())
@@ -450,9 +482,9 @@ FTransform UMotionWarpingComponent::ProcessRootMotionPreConvertToWorld(const FTr
 				OriginalRootMotionAccum = ActorFeetLocation;
 				WarpedRootMotionAccum = ActorFeetLocation;
 			}
-
-			OriginalRootMotionAccum = OriginalRootMotionAccum.GetValue() + (CharacterOwner->GetMesh()->ConvertLocalRootMotionToWorld(FTransform(InRootMotion.GetLocation()))).GetLocation();
-			WarpedRootMotionAccum = WarpedRootMotionAccum.GetValue() + (CharacterOwner->GetMesh()->ConvertLocalRootMotionToWorld(FTransform(FinalRootMotion.GetLocation()))).GetLocation();
+			
+			OriginalRootMotionAccum = OriginalRootMotionAccum.GetValue() + (OwnerAdapter->GetMesh()->ConvertLocalRootMotionToWorld(FTransform(InRootMotion.GetLocation()))).GetLocation();
+			WarpedRootMotionAccum = WarpedRootMotionAccum.GetValue() + (OwnerAdapter->GetMesh()->ConvertLocalRootMotionToWorld(FTransform(FinalRootMotion.GetLocation()))).GetLocation();
 
 			DrawDebugPoint(GetWorld(), OriginalRootMotionAccum.GetValue(), PointSize, FColor::Red, false, DrawDebugDuration, 0);
 			DrawDebugPoint(GetWorld(), WarpedRootMotionAccum.GetValue(), PointSize, FColor::Green, false, DrawDebugDuration, 0);
@@ -470,6 +502,7 @@ FTransform UMotionWarpingComponent::ProcessRootMotionPreConvertToWorld(const FTr
 	return FinalRootMotion;
 }
 
+
 bool UMotionWarpingComponent::FindAndUpdateWarpTarget(const FMotionWarpingTarget& WarpTarget)
 {
 	for (int32 Idx = 0; Idx < WarpTargets.Num(); Idx++)
diff --git a/Engine/Plugins/Animation/MotionWarping/Source/MotionWarping/Private/RootMotionModifier.cpp b/Engine/Plugins/Animation/MotionWarping/Source/MotionWarping/Private/RootMotionModifier.cpp
--- a/Engine/Plugins/Animation/MotionWarping/Source/MotionWarping/Private/RootMotionModifier.cpp
+++ b/Engine/Plugins/Animation/MotionWarping/Source/MotionWarping/Private/RootMotionModifier.cpp
@@ -6,6 +6,7 @@
 #include "Components/SkeletalMeshComponent.h"
 #include "Engine/World.h"
 #include "MotionWarpingComponent.h"
+#include "MotionWarpingAdapter.h"
 #include "DrawDebugHelpers.h"
 
 #include UE_INLINE_GENERATED_CPP_BY_NAME(RootMotionModifier)
@@ -93,16 +94,32 @@ UMotionWarpingComponent* URootMotionModifier::GetOwnerComponent() const
 	return Cast<UMotionWarpingComponent>(GetOuter());
 }
 
-ACharacter* URootMotionModifier::GetCharacterOwner() const
+UMotionWarpingBaseAdapter* URootMotionModifier::GetOwnerAdapter() const
 {
 	UMotionWarpingComponent* OwnerComp = GetOwnerComponent();
-	return OwnerComp ? OwnerComp->GetCharacterOwner() : nullptr;
+	return OwnerComp ? OwnerComp->GetOwnerAdapter() : nullptr;
+}
+
+AActor* URootMotionModifier::GetActorOwner() const
+{
+	if (UMotionWarpingBaseAdapter* OwnerAdapter = GetOwnerAdapter())
+	{
+		return OwnerAdapter->GetActor();
+	}
+
+	return nullptr;
+}
+
+ACharacter* URootMotionModifier::GetCharacterOwner() const
+{
+	return Cast<ACharacter>(GetActorOwner());
 }
 
 void URootMotionModifier::Update(const FMotionWarpingUpdateContext& Context)
 {
-	const ACharacter* CharacterOwner = GetCharacterOwner();
-	if (CharacterOwner == nullptr)
+	const AActor* ActorOwner = GetActorOwner();
+
+	if (ActorOwner == nullptr)
 	{
 		return;
 	}
@@ -111,7 +128,7 @@ void URootMotionModifier::Update(const FMotionWarpingUpdateContext& Context)
 	if (!Context.Animation.IsValid() || Context.Animation.Get() != Animation)
 	{
 		UE_LOG(LogMotionWarping, Verbose, TEXT("MotionWarping: Marking RootMotionModifier for removal. Reason: Animation is not valid. Char: %s Current Animation: %s. Window: Animation: %s [%f %f] [%f %f]"),
-			*GetNameSafe(CharacterOwner), *GetNameSafe(Context.Animation.Get()), *GetNameSafe(Animation.Get()), StartTime, EndTime, PreviousPosition, CurrentPosition);
+			*GetNameSafe(ActorOwner), *GetNameSafe(Context.Animation.Get()), *GetNameSafe(Animation.Get()), StartTime, EndTime, PreviousPosition, CurrentPosition);
 
 		SetState(ERootMotionModifierState::MarkedForRemoval);
 		return;
@@ -126,7 +143,7 @@ void URootMotionModifier::Update(const FMotionWarpingUpdateContext& Context)
 	if (PreviousPosition >= EndTime)
 	{
 		UE_LOG(LogMotionWarping, Verbose, TEXT("MotionWarping: Marking RootMotionModifier for removal. Reason: Window has ended. Char: %s Animation: %s [%f %f] [%f %f]"),
-			*GetNameSafe(CharacterOwner), *GetNameSafe(Animation.Get()), StartTime, EndTime, PreviousPosition, CurrentPosition);
+			*GetNameSafe(ActorOwner), *GetNameSafe(Animation.Get()), StartTime, EndTime, PreviousPosition, CurrentPosition);
 
 		SetState(ERootMotionModifierState::MarkedForRemoval);
 		return;
@@ -184,14 +201,15 @@ void URootMotionModifier::OnStateChanged(ERootMotionModifierState LastState)
 	{
 		if (LastState != ERootMotionModifierState::Active && State == ERootMotionModifierState::Active)
 		{
-			const ACharacter* CharacterOwner = OwnerComp->GetCharacterOwner();
-			check(CharacterOwner);
+			const UMotionWarpingBaseAdapter* OwnerAdapter = GetOwnerAdapter();
+
+			checkf(OwnerAdapter, TEXT("Root motion modifiers expect an owner and adapter"));
+
+			const FVector CurrentLocation = OwnerAdapter->GetVisualRootLocation();
+			const FQuat CurrentRotation = OwnerAdapter->GetActor()->GetActorQuat();
 			
 			ActualStartTime = PreviousPosition;
 
-			const float CapsuleHalfHeight = CharacterOwner->GetCapsuleComponent()->GetScaledCapsuleHalfHeight();
-			const FQuat CurrentRotation = CharacterOwner->GetActorQuat();
-			const FVector CurrentLocation = (CharacterOwner->GetActorLocation() - CurrentRotation.GetUpVector() * CapsuleHalfHeight);
 			StartTransform = FTransform(CurrentRotation, CurrentLocation);
 
 			TotalRootMotionWithinWindow = UMotionWarpingUtilities::ExtractRootMotionFromAnimation(Animation.Get(), StartTime, EndTime);
@@ -243,15 +261,15 @@ void URootMotionModifier_Warp::Update(const FMotionWarpingUpdateContext& Context
 		{
 			if (!CachedOffsetFromWarpPoint.IsSet())
 			{
-				if (const ACharacter* CharacterOwner = GetCharacterOwner())
+				if (const UMotionWarpingBaseAdapter* OwnerAdapter = GetOwnerAdapter())
 				{
 					if (WarpPointAnimProvider == EWarpPointAnimProvider::Static)
 					{
-						CachedOffsetFromWarpPoint = UMotionWarpingUtilities::CalculateRootTransformRelativeToWarpPointAtTime(*CharacterOwner, GetAnimation(), EndTime, WarpPointAnimTransform);
+						CachedOffsetFromWarpPoint = UMotionWarpingUtilities::CalculateRootTransformRelativeToWarpPointAtTime(*OwnerAdapter, GetAnimation(), EndTime, WarpPointAnimTransform);
 					}
 					else if (WarpPointAnimProvider == EWarpPointAnimProvider::Bone)
 					{
-						CachedOffsetFromWarpPoint = UMotionWarpingUtilities::CalculateRootTransformRelativeToWarpPointAtTime(*CharacterOwner, GetAnimation(), EndTime, WarpPointAnimBoneName);
+						CachedOffsetFromWarpPoint = UMotionWarpingUtilities::CalculateRootTransformRelativeToWarpPointAtTime(*OwnerAdapter, GetAnimation(), EndTime, WarpPointAnimBoneName);
 					}
 				}
 			}
@@ -271,13 +289,12 @@ void URootMotionModifier_Warp::Update(const FMotionWarpingUpdateContext& Context
 
 void URootMotionModifier_Warp::OnTargetTransformChanged()
 {
-	if (const ACharacter* CharacterOwner = GetCharacterOwner())
+ 	if (const UMotionWarpingBaseAdapter* WarpingAdapter = GetOwnerAdapter())
 	{
 		ActualStartTime = PreviousPosition;
 
-		const float CapsuleHalfHeight = CharacterOwner->GetCapsuleComponent()->GetScaledCapsuleHalfHeight();
-		const FQuat CurrentRotation = CharacterOwner->GetActorQuat();
-		const FVector CurrentLocation = (CharacterOwner->GetActorLocation() - CurrentRotation.GetUpVector() * CapsuleHalfHeight);
+		const FQuat CurrentRotation = WarpingAdapter->GetActor()->GetActorQuat();
+		const FVector CurrentLocation = WarpingAdapter->GetVisualRootLocation();
 		StartTransform = FTransform(CurrentRotation, CurrentLocation);
 	}
 }
@@ -290,10 +307,10 @@ FQuat URootMotionModifier_Warp::GetTargetRotation() const
 	}
 	else if (RotationType == EMotionWarpRotationType::Facing)
 	{
-		if (const ACharacter* CharacterOwner = GetCharacterOwner())
+		if (const AActor* ActorOwner = GetActorOwner())
 		{
-			const FTransform& CharacterTransform = CharacterOwner->GetActorTransform();
-			const FVector ToSyncPoint = (CachedTargetTransform.GetLocation() - CharacterTransform.GetLocation()).GetSafeNormal2D();
+			const FTransform& ActorTransform = ActorOwner->GetActorTransform();
+			const FVector ToSyncPoint = (CachedTargetTransform.GetLocation() - ActorTransform.GetLocation()).GetSafeNormal2D();
 			return FRotationMatrix::MakeFromXZ(ToSyncPoint, FVector::UpVector).ToQuat();
 		}
 	}
@@ -303,15 +320,21 @@ FQuat URootMotionModifier_Warp::GetTargetRotation() const
 
 FQuat URootMotionModifier_Warp::WarpRotation(const FTransform& RootMotionDelta, const FTransform& RootMotionTotal, float DeltaSeconds)
 {
-	const ACharacter* CharacterOwner = GetCharacterOwner();
-	if (CharacterOwner == nullptr)
+	FQuat CurrentRotation;
+	FQuat TargetRotation;
+
+	if (const UMotionWarpingBaseAdapter* WarpingAdapter = GetOwnerAdapter())
 	{
+		CurrentRotation = WarpingAdapter->GetActor()->GetActorQuat() * WarpingAdapter->GetBaseVisualRotationOffset();
+		TargetRotation = CurrentRotation.Inverse() * (GetTargetRotation() * WarpingAdapter->GetBaseVisualRotationOffset());
+	}
+	else
+	{
+		// No owner, no warping possible
 		return FQuat::Identity;
 	}
 
 	const FQuat TotalRootMotionRotation = RootMotionTotal.GetRotation();
-	const FQuat CurrentRotation = CharacterOwner->GetActorQuat() * CharacterOwner->GetBaseRotationOffset();
-	const FQuat TargetRotation = CurrentRotation.Inverse() * (GetTargetRotation() * CharacterOwner->GetBaseRotationOffset());
 	const float TimeRemaining = (EndTime - PreviousPosition) * WarpRotationTimeMultiplier;
 	const float Alpha = FMath::Clamp(DeltaSeconds / TimeRemaining, 0.f, 1.f);
 	FQuat TargetRotThisFrame = FQuat::Slerp(TotalRootMotionRotation, TargetRotation, Alpha);
@@ -342,23 +365,33 @@ FQuat URootMotionModifier_Warp::WarpRotation(const FTransform& RootMotionDelta,
 #if !(UE_BUILD_SHIPPING || UE_BUILD_TEST)
 void URootMotionModifier_Warp::PrintLog(const FString& Name, const FTransform& OriginalRootMotion, const FTransform& WarpedRootMotion) const
 {
-	if (const ACharacter* CharacterOwner = GetCharacterOwner())
+	const AActor* ActorOwner = nullptr;
+
+	FVector CurrentLocation;
+	USkeletalMeshComponent* SkelMesh = nullptr;
+
+	if (const UMotionWarpingBaseAdapter* WarpingAdapter = GetOwnerAdapter())
+	{
+		ActorOwner = WarpingAdapter->GetActor();
+		SkelMesh = WarpingAdapter->GetMesh();
+		CurrentLocation = WarpingAdapter->GetVisualRootLocation();
+	}
+
+	if (ActorOwner && SkelMesh)
 	{
-		const float CapsuleHalfHeight = CharacterOwner->GetCapsuleComponent()->GetScaledCapsuleHalfHeight();
-		const FVector CurrentLocation = (CharacterOwner->GetActorLocation() - FVector(0.f, 0.f, CapsuleHalfHeight));
 		const FVector CurrentToTarget = (GetTargetLocation() - CurrentLocation).GetSafeNormal2D();
-		const FVector FutureLocation = CurrentLocation + (CharacterOwner->GetMesh()->ConvertLocalRootMotionToWorld(WarpedRootMotion)).GetTranslation();
-		const FRotator CurrentRotation = CharacterOwner->GetActorRotation();
-		const FRotator FutureRotation = (WarpedRootMotion.GetRotation() * CharacterOwner->GetActorQuat()).Rotator();
-		const float Dot = FVector::DotProduct(CharacterOwner->GetActorForwardVector(), CurrentToTarget);
+		const FVector FutureLocation = CurrentLocation + (SkelMesh->ConvertLocalRootMotionToWorld(WarpedRootMotion)).GetTranslation();
+		const FRotator CurrentRotation = ActorOwner->GetActorRotation();
+		const FRotator FutureRotation = (WarpedRootMotion.GetRotation() * ActorOwner->GetActorQuat()).Rotator();
+		const float Dot = FVector::DotProduct(ActorOwner->GetActorForwardVector(), CurrentToTarget);
 		const float CurrentDist2D = FVector::Dist2D(GetTargetLocation(), CurrentLocation);
 		const float FutureDist2D = FVector::Dist2D(GetTargetLocation(), FutureLocation);
-		const float DeltaSeconds = CharacterOwner->GetWorld()->GetDeltaSeconds();
+		const float DeltaSeconds = ActorOwner->GetWorld()->GetDeltaSeconds();
 		const float Speed = WarpedRootMotion.GetTranslation().Size() / DeltaSeconds;
 		const float EndTimeOffset = CurrentPosition - EndTime;
 
 		UE_LOG(LogMotionWarping, Log, TEXT("%s NetMode: %d Char: %s Anim: %s Win: [%f %f][%f %f] DT: %f WT: %f ETOffset: %f Dist2D: %f Z: %f FDist2D: %f FZ: %f Dot: %f Delta: %s (%f) FDelta: %s (%f) Speed: %f Loc: %s FLoc: %s Rot: %s FRot: %s"),
-			*Name, (int32)CharacterOwner->GetWorld()->GetNetMode(), *GetNameSafe(CharacterOwner), *GetNameSafe(Animation.Get()), StartTime, EndTime, PreviousPosition, CurrentPosition, DeltaSeconds, CharacterOwner->GetWorld()->GetTimeSeconds(), EndTimeOffset,
+			*Name, (int32)ActorOwner->GetWorld()->GetNetMode(), *GetNameSafe(ActorOwner), *GetNameSafe(Animation.Get()), StartTime, EndTime, PreviousPosition, CurrentPosition, DeltaSeconds, ActorOwner->GetWorld()->GetTimeSeconds(), EndTimeOffset,
 			CurrentDist2D, (GetTargetLocation().Z - CurrentLocation.Z), FutureDist2D, (GetTargetLocation().Z - FutureLocation.Z), Dot,
 			*OriginalRootMotion.GetTranslation().ToString(), OriginalRootMotion.GetTranslation().Size(), *WarpedRootMotion.GetTranslation().ToString(), WarpedRootMotion.GetTranslation().Size(), Speed,
 			*CurrentLocation.ToString(), *FutureLocation.ToString(), *CurrentRotation.ToCompactString(), *FutureRotation.ToCompactString());
@@ -376,7 +409,13 @@ UDEPRECATED_RootMotionModifier_SimpleWarp::UDEPRECATED_RootMotionModifier_Simple
 
 FTransform UDEPRECATED_RootMotionModifier_SimpleWarp::ProcessRootMotion(const FTransform& InRootMotion, float DeltaSeconds)
 {
-	const ACharacter* CharacterOwner = GetCharacterOwner();
+	const ACharacter* CharacterOwner = nullptr;
+
+	if (const UMotionWarpingBaseAdapter* Adapter = GetOwnerAdapter())
+	{
+		CharacterOwner = Cast<ACharacter>(Adapter->GetActor());
+	}
+
 	if (CharacterOwner == nullptr)
 	{
 		return InRootMotion;
diff --git a/Engine/Plugins/Animation/MotionWarping/Source/MotionWarping/Private/RootMotionModifier_AdjustmentBlendWarp.cpp b/Engine/Plugins/Animation/MotionWarping/Source/MotionWarping/Private/RootMotionModifier_AdjustmentBlendWarp.cpp
--- a/Engine/Plugins/Animation/MotionWarping/Source/MotionWarping/Private/RootMotionModifier_AdjustmentBlendWarp.cpp
+++ b/Engine/Plugins/Animation/MotionWarping/Source/MotionWarping/Private/RootMotionModifier_AdjustmentBlendWarp.cpp
@@ -9,6 +9,7 @@
 #include "DrawDebugHelpers.h"
 #include "GameFramework/Character.h"
 #include "MotionWarpingComponent.h"
+#include "MotionWarpingAdapter.h"
 
 #include UE_INLINE_GENERATED_CPP_BY_NAME(RootMotionModifier_AdjustmentBlendWarp)
 
@@ -54,8 +55,12 @@ void URootMotionModifier_AdjustmentBlendWarp::ExtractBoneTransformAtFrame(FTrans
 
 void URootMotionModifier_AdjustmentBlendWarp::OnTargetTransformChanged()
 {
-	const ACharacter* CharacterOwner = GetCharacterOwner();
-	const USkeletalMeshComponent* SkelMeshComp = CharacterOwner ? CharacterOwner->GetMesh() : nullptr;
+	const USkeletalMeshComponent* SkelMeshComp = nullptr;
+
+	if (UMotionWarpingBaseAdapter* OwnerAdapter = GetOwnerAdapter())
+	{
+		SkelMeshComp = OwnerAdapter->GetMesh();
+	}
 
 	if (SkelMeshComp)
 	{
@@ -104,13 +109,19 @@ void URootMotionModifier_AdjustmentBlendWarp::PrecomputeWarpedTracks()
 
 	// First, extract pose at the end of the window for the bones we are going to warp
 
-	const ACharacter* CharacterOwner = GetCharacterOwner();
-	if (CharacterOwner == nullptr)
+	const USkeletalMeshComponent* SkelMesh = nullptr;
+
+	if (const UMotionWarpingBaseAdapter* OwnerAdapter = GetOwnerAdapter())
+	{
+		SkelMesh = OwnerAdapter->GetMesh();
+	}
+
+	if (!SkelMesh)
 	{
 		return;
 	}
 
-	const FBoneContainer& BoneContainer = CharacterOwner->GetMesh()->GetAnimInstance()->GetRequiredBones();
+	const FBoneContainer& BoneContainer = SkelMesh->GetAnimInstance()->GetRequiredBones();
 
 	// Init FBoneContainer with only the bones that we are interested in
 	TArray<FBoneIndexType> RequiredBoneIndexArray;
@@ -353,8 +364,15 @@ FTransform URootMotionModifier_AdjustmentBlendWarp::ExtractWarpedRootMotion() co
 #if !(UE_BUILD_SHIPPING || UE_BUILD_TEST)
 void URootMotionModifier_AdjustmentBlendWarp::DrawDebugWarpedTracks(float DrawDuration) const
 {
-	const ACharacter* CharacterOwner = GetCharacterOwner();
-	const UWorld* World = CharacterOwner ? CharacterOwner->GetWorld() : nullptr;
+	const UWorld* World = nullptr;
+	const USkeletalMeshComponent* SkelMesh = nullptr;
+
+	if (const UMotionWarpingBaseAdapter* OwnerAdapter = GetOwnerAdapter())
+	{
+		World = OwnerAdapter->GetActor()->GetWorld();
+		SkelMesh = OwnerAdapter->GetMesh();
+	}
+
 	if (World && Result.GetNum() > 0 && PreviousPosition <= EndTime)
 	{
 		FTransform RootStartTransform;
@@ -380,7 +398,7 @@ void URootMotionModifier_AdjustmentBlendWarp::DrawDebugWarpedTracks(float DrawDu
 			}
 		}
 
-		const FTransform& MeshTransform = CharacterOwner->GetMesh()->GetComponentTransform();
+		const FTransform& MeshTransform = SkelMesh->GetComponentTransform();
 		DrawDebugCoordinateSystem(World, CachedMeshTransform.GetLocation(), CachedMeshTransform.Rotator(), 20.f, false, DrawDuration, 0, 1.f);
 		DrawDebugCoordinateSystem(World, MeshTransform.GetLocation(), MeshTransform.Rotator(), 20.f, false, DrawDuration, 0, 1.f);
 		DrawDebugCoordinateSystem(World, GetTargetLocation(), GetTargetRotator(), 50.f, false, DrawDuration, 0, 1.f);
diff --git a/Engine/Plugins/Animation/MotionWarping/Source/MotionWarping/Private/RootMotionModifier_SkewWarp.cpp b/Engine/Plugins/Animation/MotionWarping/Source/MotionWarping/Private/RootMotionModifier_SkewWarp.cpp
--- a/Engine/Plugins/Animation/MotionWarping/Source/MotionWarping/Private/RootMotionModifier_SkewWarp.cpp
+++ b/Engine/Plugins/Animation/MotionWarping/Source/MotionWarping/Private/RootMotionModifier_SkewWarp.cpp
@@ -5,6 +5,7 @@
 #include "GameFramework/Character.h"
 #include "Components/CapsuleComponent.h"
 #include "MotionWarpingComponent.h"
+#include "MotionWarpingAdapter.h"
 #if WITH_EDITOR
 #include "Animation/AnimInstance.h"
 #include "Animation/DebugSkelMeshComponent.h"
@@ -131,8 +132,15 @@ FVector URootMotionModifier_SkewWarp::WarpTranslation(const FTransform& CurrentT
 
 FTransform URootMotionModifier_SkewWarp::ProcessRootMotion(const FTransform& InRootMotion, float DeltaSeconds)
 {
-	const ACharacter* CharacterOwner = GetCharacterOwner();
-	if(CharacterOwner == nullptr)
+	const UMotionWarpingBaseAdapter* OwnerAdapter = GetOwnerAdapter();
+	const AActor* OwnerAsActor = nullptr;
+
+	if (OwnerAdapter)
+	{
+		OwnerAsActor = OwnerAdapter->GetActor();
+	}
+
+	if (!OwnerAdapter || !OwnerAsActor)
 	{
 		return InRootMotion;
 	}
@@ -150,9 +158,10 @@ FTransform URootMotionModifier_SkewWarp::ProcessRootMotion(const FTransform& InR
 
 	if (bWarpTranslation)
 	{
-		const float CapsuleHalfHeight = CharacterOwner->GetCapsuleComponent()->GetScaledCapsuleHalfHeight();
-		const FQuat CurrentRotation = CharacterOwner->GetActorQuat();
-		const FVector CurrentLocation = (CharacterOwner->GetActorLocation() - CurrentRotation.GetUpVector() * CapsuleHalfHeight);
+		const FVector CurrentLocation       = OwnerAdapter->GetVisualRootLocation();
+		const FQuat CurrentRotation         = OwnerAdapter->GetActor()->GetActorQuat();
+		const FVector MeshTranslationOffset = OwnerAdapter->GetBaseVisualTranslationOffset();
+		const FQuat MeshRotationOffset      = OwnerAdapter->GetBaseVisualRotationOffset();
 
 		const FVector DeltaTranslation = RootMotionDelta.GetLocation();
 		const FVector TotalTranslation = RootMotionTotal.GetLocation();
@@ -168,7 +177,7 @@ FTransform URootMotionModifier_SkewWarp::ProcessRootMotion(const FTransform& InR
 		{
 			if (!DeltaTranslation.IsNearlyZero())
 			{
-				const FTransform MeshTransform = FTransform(CharacterOwner->GetBaseRotationOffset(), CharacterOwner->GetBaseTranslationOffset()) * CharacterOwner->GetActorTransform();
+				const FTransform MeshTransform = FTransform(MeshRotationOffset, MeshTranslationOffset) * OwnerAsActor->GetActorTransform();
 				TargetLocation = MeshTransform.InverseTransformPositionNoScale(TargetLocation);
 
 				const FVector WarpedTranslation = WarpTranslation(FTransform::Identity, DeltaTranslation, TotalTranslation, TargetLocation) + ExtraRootMotion.GetLocation();
@@ -191,7 +200,7 @@ FTransform URootMotionModifier_SkewWarp::ProcessRootMotion(const FTransform& InR
 				const FVector NextLocation = FMath::Lerp<FVector, float>(StartTransform.GetLocation(), TargetLocation, Alpha);
 				FVector FinalDeltaTranslation = (NextLocation - CurrentLocation);
 				FinalDeltaTranslation = (CurrentRotation.Inverse() * DeltaToTarget.ToOrientationQuat()).GetForwardVector() * FinalDeltaTranslation.Size();
-				FinalDeltaTranslation = CharacterOwner->GetBaseRotationOffset().UnrotateVector(FinalDeltaTranslation);
+				FinalDeltaTranslation = MeshRotationOffset.UnrotateVector(FinalDeltaTranslation);
 
 				FinalRootMotion.SetTranslation(FinalDeltaTranslation + ExtraRootMotion.GetLocation());
 			}
@@ -215,7 +224,7 @@ FTransform URootMotionModifier_SkewWarp::ProcessRootMotion(const FTransform& InR
 	if (DebugLevel == 2 || DebugLevel == 3)
 	{
 		const float DrawDebugDuration = FMotionWarpingCVars::CVarMotionWarpingDrawDebugDuration.GetValueOnGameThread();
-		DrawDebugCoordinateSystem(CharacterOwner->GetWorld(), GetTargetLocation(), GetTargetRotator(), 50.f, false, DrawDebugDuration, 0, 1.f);
+		DrawDebugCoordinateSystem(OwnerAsActor->GetWorld(), GetTargetLocation(), GetTargetRotator(), 50.f, false, DrawDebugDuration, 0, 1.f);
 	}
 #endif
 
diff --git a/Engine/Plugins/Animation/MotionWarping/Source/MotionWarping/Public/MotionWarpingAdapter.h b/Engine/Plugins/Animation/MotionWarping/Source/MotionWarping/Public/MotionWarpingAdapter.h
new file mode 100644
--- /dev/null
+++ b/Engine/Plugins/Animation/MotionWarping/Source/MotionWarping/Public/MotionWarpingAdapter.h
@@ -0,0 +1,31 @@
+// Copyright Epic Games, Inc. All Rights Reserved.
+
+#pragma once
+
+#include "RootMotionModifier.h"
+#include "MotionWarpingAdapter.generated.h"
+
+
+DECLARE_DELEGATE_RetVal_ThreeParams(FTransform, FOnWarpLocalspaceRootMotionWithContext, const FTransform&, float, const FMotionWarpingUpdateContext*)
+DECLARE_DELEGATE_RetVal_ThreeParams(FTransform, FOnWarpWorldspaceRootMotionWithContext, const FTransform&, float, const FMotionWarpingUpdateContext*)
+
+/**
+ * MotionWarpingBaseAdapter: base class to adapt/apply motion warping to a target. Concrete subclasses should override
+ */
+UCLASS(Abstract)
+class MOTIONWARPING_API UMotionWarpingBaseAdapter : public UObject
+{
+	GENERATED_BODY()
+
+public:
+	virtual ~UMotionWarpingBaseAdapter() {}
+	virtual AActor* GetActor() const { return nullptr; }
+	virtual USkeletalMeshComponent* GetMesh() const { return nullptr; }
+	virtual FVector GetVisualRootLocation() const { return FVector::ZeroVector; }
+	virtual FVector GetBaseVisualTranslationOffset() const { return FVector::ZeroVector; }
+	virtual FQuat GetBaseVisualRotationOffset() const { return FQuat::Identity; }
+
+	// A MotionWarpingComponent will bind to this delegate to perform warping when it is triggered
+	FOnWarpLocalspaceRootMotionWithContext WarpLocalRootMotionDelegate;
+};
+
diff --git a/Engine/Plugins/Animation/MotionWarping/Source/MotionWarping/Public/MotionWarpingCharacterAdapter.h b/Engine/Plugins/Animation/MotionWarping/Source/MotionWarping/Public/MotionWarpingCharacterAdapter.h
new file mode 100644
--- /dev/null
+++ b/Engine/Plugins/Animation/MotionWarping/Source/MotionWarping/Public/MotionWarpingCharacterAdapter.h
@@ -0,0 +1,34 @@
+// Copyright Epic Games, Inc. All Rights Reserved.
+
+#pragma once
+
+#include "MotionWarpingAdapter.h"
+#include "GameFramework/CharacterMovementComponent.h"
+#include "MotionWarpingCharacterAdapter.generated.h"
+
+// Adapter for Character / ChararacterMovementComponent actors to participate in motion warping
+UCLASS()
+class MOTIONWARPING_API UMotionWarpingCharacterAdapter : public UMotionWarpingBaseAdapter
+{
+	GENERATED_BODY()
+
+public:
+	UMotionWarpingCharacterAdapter();
+
+	virtual ~UMotionWarpingCharacterAdapter();
+
+	void SetCharacter(ACharacter* InCharacter);
+
+	virtual AActor* GetActor() const override;
+	virtual USkeletalMeshComponent* GetMesh() const override;
+	virtual FVector GetVisualRootLocation() const override;
+	virtual FVector GetBaseVisualTranslationOffset() const override;
+	virtual FQuat GetBaseVisualRotationOffset() const override;
+
+private:
+	// Triggered when the character says it's time to pre-process local root motion. This adapter catches the request and passes along to the Warping component
+	FTransform WarpLocalRootMotionOnCharacter(const FTransform& LocalRootMotionTransform, UCharacterMovementComponent* TargetMoveComp, float DeltaSeconds);
+
+	/** The associated character */
+	TObjectPtr<ACharacter> TargetCharacter;
+};
\ No newline at end of file
diff --git a/Engine/Plugins/Animation/MotionWarping/Source/MotionWarping/Public/MotionWarpingComponent.h b/Engine/Plugins/Animation/MotionWarping/Source/MotionWarping/Public/MotionWarpingComponent.h
--- a/Engine/Plugins/Animation/MotionWarping/Source/MotionWarping/Public/MotionWarpingComponent.h
+++ b/Engine/Plugins/Animation/MotionWarping/Source/MotionWarping/Public/MotionWarpingComponent.h
@@ -11,6 +11,7 @@ template <class T> class TAutoConsoleVariable;
 class ACharacter;
 class UAnimSequenceBase;
 class UCharacterMovementComponent;
+class UMotionWarpingBaseAdapter;
 class UMotionWarpingComponent;
 class UAnimNotifyState_MotionWarping;
 struct FCompactPose;
@@ -42,6 +43,7 @@ struct FMotionWarpingWindowData
 	float EndTime = 0.f;
 };
 
+
 UCLASS()
 class MOTIONWARPING_API UMotionWarpingUtilities : public UBlueprintFunctionLibrary
 {
@@ -72,11 +74,14 @@ public:
 
 	/** @return root transform relative to the warp point bone at the supplied time */
 	static FTransform CalculateRootTransformRelativeToWarpPointAtTime(const ACharacter& Character, const UAnimSequenceBase* Animation, float Time, const FName& WarpPointBoneName);
+	static FTransform CalculateRootTransformRelativeToWarpPointAtTime(const UMotionWarpingBaseAdapter& WarpingAdapter, const UAnimSequenceBase* Animation, float Time, const FName& WarpPointBoneName);
 
 	/** @return root transform relative to the warp point transform at the supplied time */
 	static FTransform CalculateRootTransformRelativeToWarpPointAtTime(const ACharacter& Character, const UAnimSequenceBase* Animation, float Time, const FTransform& WarpPointTransform);
+	static FTransform CalculateRootTransformRelativeToWarpPointAtTime(const UMotionWarpingBaseAdapter& WarpingAdapter, const UAnimSequenceBase* Animation, float Time, const FTransform& WarpPointTransform);
 };
 
+
 DECLARE_DYNAMIC_MULTICAST_DELEGATE_OneParam(FMotionWarpingPreUpdate, class UMotionWarpingComponent*, MotionWarpingComp);
 
 UCLASS(ClassGroup = Movement, meta = (BlueprintSpawnableComponent))
@@ -99,8 +104,22 @@ public:
 	virtual void InitializeComponent() override;
 	virtual void GetLifetimeReplicatedProps(TArray< FLifetimeProperty >& OutLifetimeProps) const;
 
-	/** Gets the character this component belongs to */
-	FORCEINLINE ACharacter* GetCharacterOwner() const { return CharacterOwner.Get(); }
+	/** Set a new adapter of a particular type */
+	UMotionWarpingBaseAdapter* CreateOwnerAdapter(TSubclassOf<UMotionWarpingBaseAdapter> AdapterClass);
+
+	template <class T>
+	T* CreateOwnerAdapter()
+	{
+		static_assert(TPointerIsConvertibleFromTo<T, const UMotionWarpingBaseAdapter>::Value, "'T' template parameter to CreateOwnerAdapter must be derived from UMotionWarpingBaseAdapter");
+		return CastChecked<T>(CreateOwnerAdapter(T::StaticClass()));
+	}
+
+	/** Get the current adapter to the owner */
+	UMotionWarpingBaseAdapter* GetOwnerAdapter() const { return OwnerAdapter; }
+
+	/** Gets the Character this component belongs to. Returns null if not owned by a Character actor. */
+	UE_DEPRECATED(5.5, "Motion Warping is no longer limited to Character actors. Use GetOwnerAdapter()->GetActor() instead.")
+	ACharacter* GetCharacterOwner() const;
 
 	/** Returns the list of root motion modifiers */
 	FORCEINLINE const TArray<URootMotionModifier*>& GetModifiers() const { return Modifiers; }
@@ -180,10 +199,10 @@ public:
 
 protected:
 
-	/** Character this component belongs to */
+	/** Adapter that connects motion warping to an owner */
 	UPROPERTY(Transient)
-	TWeakObjectPtr<ACharacter> CharacterOwner;
-
+	TObjectPtr<UMotionWarpingBaseAdapter> OwnerAdapter;
+	
 	/** List of root motion modifiers */
 	UPROPERTY(Transient)
 	TArray<TObjectPtr<URootMotionModifier>> Modifiers;
@@ -196,11 +215,11 @@ protected:
 	TOptional<FVector> WarpedRootMotionAccum;
 #endif
 
-	void Update(float DeltaSeconds);
+	void UpdateWithContext(const FMotionWarpingUpdateContext& Context, float DeltaSeconds);
 
 	bool FindAndUpdateWarpTarget(const FMotionWarpingTarget& WarpTarget);
 
-	FTransform ProcessRootMotionPreConvertToWorld(const FTransform& InRootMotion, class UCharacterMovementComponent* CharacterMovementComponent, float DeltaSeconds);
+	FTransform ProcessRootMotionPreConvertToWorld(const FTransform& InRootMotion, float DeltaSeconds, const FMotionWarpingUpdateContext* InContext=nullptr); // callback with optional context
 };
 
 #if UE_ENABLE_INCLUDE_ORDER_DEPRECATED_IN_5_2
diff --git a/Engine/Plugins/Animation/MotionWarping/Source/MotionWarping/Public/RootMotionModifier.h b/Engine/Plugins/Animation/MotionWarping/Source/MotionWarping/Public/RootMotionModifier.h
--- a/Engine/Plugins/Animation/MotionWarping/Source/MotionWarping/Public/RootMotionModifier.h
+++ b/Engine/Plugins/Animation/MotionWarping/Source/MotionWarping/Public/RootMotionModifier.h
@@ -10,6 +10,8 @@ class UMotionWarpingComponent;
 class UAnimNotifyState_MotionWarping;
 class URootMotionModifier;
 class USceneComponent;
+class ACharacter;
+class UMotionWarpingBaseAdapter;
 
 /** 
  * Context passed to any active root motion modifier during the update phase. 
@@ -130,8 +132,15 @@ public:
 	/** Returns a pointer to the component that owns this modifier */
 	UMotionWarpingComponent* GetOwnerComponent() const;
 
-	/** Returns a pointer to the character that owns the component that owns this modifier */
-	class ACharacter* GetCharacterOwner() const;
+	/** Returns a pointer to the adapter to the object being warped */
+	UMotionWarpingBaseAdapter* GetOwnerAdapter() const;
+
+	/** Returns a pointer to the actor that owns the component that owns this modifier */
+	AActor* GetActorOwner() const;
+	
+	/** Returns a pointer to the Character that owns the component using this modifier. Returns null if not owned by a Character actor. */
+	UE_DEPRECATED(5.5, "Motion Warping is no longer limited to Character actors. Use GetActorOwner instead.")
+	ACharacter* GetCharacterOwner() const;
 
 	virtual void Update(const FMotionWarpingUpdateContext& Context);
 	virtual FTransform ProcessRootMotion(const FTransform& InRootMotion, float DeltaSeconds) { return FTransform::Identity; }
diff --git a/Engine/Plugins/Experimental/Mover/Source/Mover/Private/DefaultMovementSet/LayeredMoves/AnimRootMotionLayeredMove.cpp b/Engine/Plugins/Experimental/Mover/Source/Mover/Private/DefaultMovementSet/LayeredMoves/AnimRootMotionLayeredMove.cpp
--- a/Engine/Plugins/Experimental/Mover/Source/Mover/Private/DefaultMovementSet/LayeredMoves/AnimRootMotionLayeredMove.cpp
+++ b/Engine/Plugins/Experimental/Mover/Source/Mover/Private/DefaultMovementSet/LayeredMoves/AnimRootMotionLayeredMove.cpp
@@ -27,6 +27,8 @@ bool FLayeredMove_AnimRootMotion::GenerateMove(const FMoverTickStartData& SimSta
 
 	const AActor* MoverActor = MoverComp->GetOwner();
 
+	const FMoverDefaultSyncState* SyncState = SimState.SyncState.SyncStateCollection.FindDataByType<FMoverDefaultSyncState>();
+
 	// First pass simply samples based on the duration. For long animations, this has the potential to diverge.
 	// Future improvements could include:
 	//     - speeding up or slowing down slightly to match the associated montage instance
@@ -42,22 +44,17 @@ bool FLayeredMove_AnimRootMotion::GenerateMove(const FMoverTickStartData& SimSta
 	// Read the local transform directly from the montage
 	const FTransform LocalRootMotion = UMotionWarpingUtilities::ExtractRootMotionFromAnimation(Montage, ExtractionStartPosition, ExtractionEndPosition);
 
-	FTransform WorldSpaceRootMotion;
-	
-	if (USkeletalMeshComponent* SkeletalMesh = MoverActor->FindComponentByClass<USkeletalMeshComponent>())
-	{
-		WorldSpaceRootMotion = SkeletalMesh->ConvertLocalRootMotionToWorld(LocalRootMotion);
-	}
-	else
-	{
-		const FTransform ActorToWorldTransform = MoverActor->GetTransform();
-		const FVector DeltaWorldTranslation = LocalRootMotion.GetTranslation() - ActorToWorldTransform.GetTranslation();
-
-		const FQuat NewWorldRotation = ActorToWorldTransform.GetRotation() * LocalRootMotion.GetRotation();
-		const FQuat DeltaWorldRotation = NewWorldRotation * ActorToWorldTransform.GetRotation().Inverse();
-
-		WorldSpaceRootMotion.SetComponents(DeltaWorldRotation, DeltaWorldTranslation, FVector::OneVector);
-	}
+	FMotionWarpingUpdateContext WarpingContext;
+	WarpingContext.Animation = Montage;
+	WarpingContext.CurrentPosition = ExtractionEndPosition;
+	WarpingContext.PreviousPosition = ExtractionStartPosition;
+	WarpingContext.PlayRate = PlayRate;
+	WarpingContext.Weight = 1.f;
+
+	// Note that we're forcing the use of the sync state's actor transform data. This is necessary when the movement simulation 
+	// is running ahead of the actor's visual representation and may be rotated differently, such as in an async physics sim.
+	const FTransform SimActorTransform = FTransform(SyncState->GetOrientation_WorldSpace().Quaternion(), SyncState->GetLocation_WorldSpace());
+	const FTransform WorldSpaceRootMotion = MoverComp->ConvertLocalRootMotionToWorld(LocalRootMotion, DeltaSeconds, &SimActorTransform, &WarpingContext);
 	
 	OutProposedMove = FProposedMove();
 	OutProposedMove.MixMode = MixMode;
diff --git a/Engine/Plugins/Experimental/Mover/Source/Mover/Private/MotionWarpingMoverAdapter.cpp b/Engine/Plugins/Experimental/Mover/Source/Mover/Private/MotionWarpingMoverAdapter.cpp
new file mode 100644
--- /dev/null
+++ b/Engine/Plugins/Experimental/Mover/Source/Mover/Private/MotionWarpingMoverAdapter.cpp
@@ -0,0 +1,80 @@
+#include "MotionWarpingMoverAdapter.h"
+#include "Components/SkeletalMeshComponent.h"
+
+UMotionWarpingMoverAdapter::UMotionWarpingMoverAdapter()
+{
+}
+
+UMotionWarpingMoverAdapter::~UMotionWarpingMoverAdapter()
+{
+	if (TargetMoverComp)
+	{
+		TargetMoverComp->ProcessLocalRootMotionDelegate.Unbind();
+	}
+}
+
+void UMotionWarpingMoverAdapter::SetMoverComp(UMoverComponent* InMoverComp)
+{
+	TargetMoverComp = InMoverComp;
+	TargetMoverComp->ProcessLocalRootMotionDelegate.BindUObject(this, &UMotionWarpingMoverAdapter::WarpLocalRootMotionOnMoverComp);
+}
+
+AActor* UMotionWarpingMoverAdapter::GetActor() const
+{
+	return TargetMoverComp->GetOwner();
+}
+
+USkeletalMeshComponent* UMotionWarpingMoverAdapter::GetMesh() const
+{
+	return TargetMoverComp->GetPrimaryVisualComponent<USkeletalMeshComponent>();
+}
+
+FVector UMotionWarpingMoverAdapter::GetVisualRootLocation() const
+{
+	if (const USceneComponent* PrimaryVisualComp = TargetMoverComp->GetPrimaryVisualComponent())
+	{
+		return PrimaryVisualComp->GetComponentLocation();
+	}
+
+	if (const USceneComponent* UpdatedComponent = TargetMoverComp->GetUpdatedComponent())
+	{
+		const FVector ActorRootLocation = UpdatedComponent->GetComponentLocation();
+		const FQuat ActorRootOrientation = UpdatedComponent->GetComponentQuat();
+		FBoxSphereBounds ActorRootBounds = UpdatedComponent->GetLocalBounds();
+
+		return ActorRootLocation - (ActorRootOrientation.GetUpVector() * ActorRootBounds.BoxExtent.Z);
+	}
+
+	return TargetMoverComp->GetOwner()->GetActorLocation();
+}
+
+FVector UMotionWarpingMoverAdapter::GetBaseVisualTranslationOffset() const
+{
+	// TODO: rework these GetBase****Offset functions once MoverComponent supports primary visual component offset (coming as part of mesh-based smoothing)
+	if (const USceneComponent* VisualComp = TargetMoverComp->GetPrimaryVisualComponent())
+	{
+		return VisualComp->GetRelativeLocation();
+	}
+
+	return FVector::ZeroVector;
+}
+
+FQuat UMotionWarpingMoverAdapter::GetBaseVisualRotationOffset() const
+{
+	if (const USceneComponent* VisualComp = TargetMoverComp->GetPrimaryVisualComponent())
+	{
+		return VisualComp->GetRelativeRotation().Quaternion();
+	}
+
+	return FQuat::Identity;
+}
+
+FTransform UMotionWarpingMoverAdapter::WarpLocalRootMotionOnMoverComp(const FTransform& LocalRootMotionTransform, float DeltaSeconds, const FMotionWarpingUpdateContext* OptionalWarpingContext)
+{
+	if (WarpLocalRootMotionDelegate.IsBound())
+	{
+		return WarpLocalRootMotionDelegate.Execute(LocalRootMotionTransform, DeltaSeconds, OptionalWarpingContext);
+	}
+
+	return LocalRootMotionTransform;
+}
diff --git a/Engine/Plugins/Experimental/Mover/Source/Mover/Private/MoverComponent.cpp b/Engine/Plugins/Experimental/Mover/Source/Mover/Private/MoverComponent.cpp
--- a/Engine/Plugins/Experimental/Mover/Source/Mover/Private/MoverComponent.cpp
+++ b/Engine/Plugins/Experimental/Mover/Source/Mover/Private/MoverComponent.cpp
@@ -4,6 +4,7 @@
 #include "MoverComponent.h"
 #include "MoverSimulationTypes.h"
 #include "MovementModeStateMachine.h"
+#include "MotionWarpingMoverAdapter.h"
 #include "DefaultMovementSet/Modes/WalkingMode.h"
 #include "DefaultMovementSet/Modes/FallingMode.h"
 #include "DefaultMovementSet/Modes/FlyingMode.h"
@@ -22,6 +23,9 @@
 #include "Misc/TransactionObjectEvent.h"
 #include "Blueprint/BlueprintExceptionInfo.h"
 #include "UObject/ObjectSaveContext.h"
+#include "Components/SkeletalMeshComponent.h"
+#include "MotionWarpingComponent.h"
+
 
 #if WITH_EDITOR
 #include "Misc/DataValidation.h"
@@ -138,12 +142,6 @@ void UMoverComponent::OnRegister()
 		}
 
 		SetUpdatedComponent(NewUpdatedComponent);
-
-		// If no primary visual component is already set, fall back to searching for any kind of mesh
-		if (!PrimaryVisualComponent && MyActor)
-		{
-			PrimaryVisualComponent = MyActor->FindComponentByClass<UMeshComponent>();
-		}
 	}
 }
 
@@ -193,6 +191,22 @@ void UMoverComponent::BeginPlay()
 {
 	Super::BeginPlay();
 
+	if (const AActor* MyActor = GetOwner())
+	{
+		// If no primary visual component is already set, fall back to searching for any kind of mesh
+		if (!PrimaryVisualComponent)
+		{
+			PrimaryVisualComponent = MyActor->FindComponentByClass<UMeshComponent>();
+		}
+
+		// Optional motion warping support
+		if (UMotionWarpingComponent* WarpingComp = MyActor->FindComponentByClass<UMotionWarpingComponent>())
+		{
+			UMotionWarpingMoverAdapter* WarpingAdapter = WarpingComp->CreateOwnerAdapter<UMotionWarpingMoverAdapter>();
+			WarpingAdapter->SetMoverComp(this);
+		}
+	}
+
 	// If an InputProducer isn't already set, check the actor and its components for one
 	if (!InputProducer)
 	{
@@ -1069,6 +1083,47 @@ bool UMoverComponent::RemoveMovementMode(FName ModeName)
 }
 
 
+FTransform UMoverComponent::ConvertLocalRootMotionToWorld(const FTransform& LocalRootMotionTransform, float DeltaSeconds, const FTransform* AlternateActorToWorld, const FMotionWarpingUpdateContext* OptionalWarpingContext) const
+{
+	// Optionally process/warp localspace root motion
+	const FTransform ProcessedLocalRootMotion = ProcessLocalRootMotionDelegate.IsBound()
+		? ProcessLocalRootMotionDelegate.Execute(LocalRootMotionTransform, DeltaSeconds, OptionalWarpingContext)
+		: LocalRootMotionTransform;
+
+	// Convert processed localspace root motion to worldspace
+	FTransform WorldSpaceRootMotion;
+
+	if (USkeletalMeshComponent* SkeletalMesh = GetPrimaryVisualComponent<USkeletalMeshComponent>())
+	{
+		 WorldSpaceRootMotion = SkeletalMesh->ConvertLocalRootMotionToWorld(ProcessedLocalRootMotion);
+	}
+	else
+	{
+		const FTransform PresentationActorToWorldTransform = GetOwner()->GetTransform();
+		const FVector DeltaWorldTranslation = ProcessedLocalRootMotion.GetTranslation() - PresentationActorToWorldTransform.GetTranslation();
+
+		const FQuat NewWorldRotation = PresentationActorToWorldTransform.GetRotation() * ProcessedLocalRootMotion.GetRotation();
+		const FQuat DeltaWorldRotation = NewWorldRotation * PresentationActorToWorldTransform.GetRotation().Inverse();
+
+		WorldSpaceRootMotion.SetComponents(DeltaWorldRotation, DeltaWorldTranslation, FVector::OneVector);
+	}
+
+	// Optionally convert this to be relative to a different space
+	if (AlternateActorToWorld)
+	{
+		const FTransform WorldToActor = GetOwner()->GetTransform().Inverse();
+		const FTransform ActorSpaceRootMotion = WorldSpaceRootMotion * WorldToActor;
+		WorldSpaceRootMotion = ActorSpaceRootMotion * *AlternateActorToWorld;
+	}
+
+	
+	// Optionally process/warp worldspace root motion
+	return ProcessWorldRootMotionDelegate.IsBound()
+		? ProcessWorldRootMotionDelegate.Execute(WorldSpaceRootMotion, DeltaSeconds, OptionalWarpingContext)
+		: WorldSpaceRootMotion;
+}
+
+
 FTransform UMoverComponent::GetUpdatedComponentTransform() const
 {
 	if (ensure(UpdatedComponent))
@@ -1079,6 +1134,12 @@ FTransform UMoverComponent::GetUpdatedComponentTransform() const
 }
 
 
+USceneComponent* UMoverComponent::GetUpdatedComponent() const
+{
+	return UpdatedComponent.Get();
+}
+
+
 USceneComponent* UMoverComponent::GetPrimaryVisualComponent() const
 {
 	return PrimaryVisualComponent.Get();
diff --git a/Engine/Plugins/Experimental/Mover/Source/Mover/Public/MotionWarpingMoverAdapter.h b/Engine/Plugins/Experimental/Mover/Source/Mover/Public/MotionWarpingMoverAdapter.h
new file mode 100644
--- /dev/null
+++ b/Engine/Plugins/Experimental/Mover/Source/Mover/Public/MotionWarpingMoverAdapter.h
@@ -0,0 +1,34 @@
+// Copyright Epic Games, Inc. All Rights Reserved.
+
+#pragma once
+
+#include "MotionWarpingAdapter.h"
+#include "MoverComponent.h"
+#include "MotionWarpingMoverAdapter.generated.h"
+
+// Adapter for MoverComponent actors to participate in motion warping
+
+UCLASS()
+class MOVER_API UMotionWarpingMoverAdapter : public UMotionWarpingBaseAdapter
+{
+	GENERATED_BODY()
+
+public:
+	UMotionWarpingMoverAdapter();
+
+	virtual ~UMotionWarpingMoverAdapter();
+
+	void SetMoverComp(UMoverComponent* InMoverComp);
+
+	virtual AActor* GetActor() const override;
+	virtual USkeletalMeshComponent* GetMesh() const override;
+	virtual FVector GetVisualRootLocation() const override;
+	virtual FVector GetBaseVisualTranslationOffset() const override;
+	virtual FQuat GetBaseVisualRotationOffset() const override;
+
+private:
+	// This is called when our Mover actor wants to warp local motion, and passes the responsibility onto the warping component
+	FTransform WarpLocalRootMotionOnMoverComp(const FTransform& LocalRootMotionTransform, float DeltaSeconds, const FMotionWarpingUpdateContext* OptionalWarpingContext);
+
+	TObjectPtr<UMoverComponent> TargetMoverComp;
+};
\ No newline at end of file
diff --git a/Engine/Plugins/Experimental/Mover/Source/Mover/Public/MoverComponent.h b/Engine/Plugins/Experimental/Mover/Source/Mover/Public/MoverComponent.h
--- a/Engine/Plugins/Experimental/Mover/Source/Mover/Public/MoverComponent.h
+++ b/Engine/Plugins/Experimental/Mover/Source/Mover/Public/MoverComponent.h
@@ -6,6 +6,7 @@
 #include "CoreMinimal.h"
 #endif // UE_ENABLE_INCLUDE_ORDER_DEPRECATED_IN_5_4
 #include "Components/ActorComponent.h"
+#include "MotionWarpingAdapter.h"
 #include "MovementMode.h"
 #include "MoverTypes.h"
 #include "LayeredMove.h"
@@ -186,11 +187,35 @@ public:
 	UFUNCTION(BlueprintCallable, BlueprintPure = false, Category = Mover)
 	FVector GetUpDirection() const;
 
+public:
+
+	/**
+	 *  Converts a local root motion transform to worldspace. 
+	 * @param AlternateActorToWorld   allows specification of a different actor root transform, for cases when root motion isn't directly being applied to this actor (async simulations)
+	 * @param OptionalWarpingContext   allows specification of a warping context, for use with root motion that is asynchronous from the actor (async simulations)
+	 */
+	virtual FTransform ConvertLocalRootMotionToWorld(const FTransform& LocalRootMotionTransform, float DeltaSeconds, const FTransform* AlternateActorToWorld=nullptr, const FMotionWarpingUpdateContext* OptionalWarpingContext=nullptr) const;
+
+	/** delegates used when converting local root motion to worldspace, allowing external systems to influence it (such as motion warping) */
+	FOnWarpLocalspaceRootMotionWithContext ProcessLocalRootMotionDelegate;
+	FOnWarpWorldspaceRootMotionWithContext ProcessWorldRootMotionDelegate;
+
 public:	// Queries
 
 	// Get the transform of the root component that our Mover simulation is moving
 	FTransform GetUpdatedComponentTransform() const;
 
+	// Access the root component of the actor that our Mover simulation is moving
+	USceneComponent* GetUpdatedComponent() const;
+
+	// Typed accessor to root moving component
+	template<class T>
+	T* GetUpdatedComponent() const
+	{
+		static_assert(TPointerIsConvertibleFromTo<T, const USceneComponent>::Value, "'T' template parameter to GetUpdatedComponent must be derived from USceneComponent");
+		return Cast<T>(GetUpdatedComponent());
+	}
+
 	// Access the primary visual component of the actor
 	USceneComponent* GetPrimaryVisualComponent() const;
 
