Motion Warping - Added function to remove a list of warp targets

diff --git a/Engine/Plugins/Animation/MotionWarping/Source/MotionWarping/Private/MotionWarpingComponent.cpp b/Engine/Plugins/Animation/MotionWarping/Source/MotionWarping/Private/MotionWarpingComponent.cpp
--- a/Engine/Plugins/Animation/MotionWarping/Source/MotionWarping/Private/MotionWarpingComponent.cpp
+++ b/Engine/Plugins/Animation/MotionWarping/Source/MotionWarping/Private/MotionWarpingComponent.cpp
@@ -543,6 +543,18 @@ int32 UMotionWarpingComponent::RemoveWarpTarget(FName WarpTargetName)
 	return NumRemoved;
 }
 
+int32 UMotionWarpingComponent::RemoveWarpTargets(const TArray<FName>& WarpTargetNames)
+{
+	const int32 NumRemoved = WarpTargets.RemoveAll([&WarpTargetNames](const FMotionWarpingTarget& WarpTarget) { return WarpTargetNames.Contains(WarpTarget.Name); });
+
+	if (NumRemoved > 0)
+	{
+		MARK_PROPERTY_DIRTY_FROM_NAME(UMotionWarpingComponent, WarpTargets, this);
+	}
+
+	return NumRemoved;
+}
+
 int32 UMotionWarpingComponent::RemoveAllWarpTargets()
 {
 	const int32 NumRemoved = WarpTargets.Num();
diff --git a/Engine/Plugins/Animation/MotionWarping/Source/MotionWarping/Public/MotionWarpingComponent.h b/Engine/Plugins/Animation/MotionWarping/Source/MotionWarping/Public/MotionWarpingComponent.h
--- a/Engine/Plugins/Animation/MotionWarping/Source/MotionWarping/Public/MotionWarpingComponent.h
+++ b/Engine/Plugins/Animation/MotionWarping/Source/MotionWarping/Public/MotionWarpingComponent.h
@@ -197,6 +197,10 @@ public:
 	UFUNCTION(BlueprintCallable, Category = "Motion Warping")
 	int32 RemoveAllWarpTargets();
 
+	/** Removes multiple warp targets */
+	UFUNCTION(BlueprintCallable, Category = "Motion Warping")
+	int32 RemoveWarpTargets(const TArray<FName>& WarpTargetNames);
+
 protected:
 
 	/** Adapter that connects motion warping to an owner */
