Optional Velocity mode for AttributeBasedRootMotionComponent

diff --git a/Engine/Plugins/Animation/MotionWarping/Source/MotionWarping/Private/AttributeBasedRootMotionComponent.cpp b/Engine/Plugins/Animation/MotionWarping/Source/MotionWarping/Private/AttributeBasedRootMotionComponent.cpp
--- a/Engine/Plugins/Animation/MotionWarping/Source/MotionWarping/Private/AttributeBasedRootMotionComponent.cpp
+++ b/Engine/Plugins/Animation/MotionWarping/Source/MotionWarping/Private/AttributeBasedRootMotionComponent.cpp
@@ -9,6 +9,7 @@
 #include "Components/SkeletalMeshComponent.h"
 #include "GameFramework/Character.h"
 #include "GameFramework/CharacterMovementComponent.h"
+#include "VisualLogger/VisualLogger.h"
 
 #include UE_INLINE_GENERATED_CPP_BY_NAME(AttributeBasedRootMotionComponent)
 
@@ -28,30 +29,111 @@ void UAttributeBasedRootMotionComponent::InitializeComponent()
 	Super::InitializeComponent();
 
 	CharacterOwner = Cast<ACharacter>(GetOwner());
+	
+	ACharacter* Character = GetCharacterOwner();
+	check(Character);
+	
+	PrePhysicsTickFunction.bCanEverTick = true;
+	PrePhysicsTickFunction.bStartWithTickEnabled = true;
+	PrePhysicsTickFunction.SetTickFunctionEnable(true);
+	PrePhysicsTickFunction.TickGroup = TG_PrePhysics;
+	PrePhysicsTickFunction.Target = this;
 }
 
-void UAttributeBasedRootMotionComponent::TickComponent(float DeltaTime, enum ELevelTick TickType, FActorComponentTickFunction* ThisTickFunction)
+void UAttributeBasedRootMotionComponent::RegisterComponentTickFunctions(bool bRegister)
 {
-	if (!bEnableRootMotion)
+	Super::RegisterComponentTickFunctions(bRegister);
+
+	if (Mode == EAttributeBasedRootMotionMode::ApplyVelocity)
 	{
-		return;
+		if (bRegister)
+		{
+			CharacterOwner = Cast<ACharacter>(GetOwner());
+			ACharacter* Character = GetCharacterOwner();
+			check(Character);
+			UCharacterMovementComponent* CharacterMovement = Character->GetCharacterMovement<UCharacterMovementComponent>();
+
+			PrePhysicsTickFunction.RegisterTickFunction(GetComponentLevel());
+			CharacterMovement->PrimaryComponentTick.AddPrerequisite(this, PrePhysicsTickFunction);
+		}
+		else
+		{
+			PrePhysicsTickFunction.UnRegisterTickFunction();
+		}
 	}
-	
+}
+
+
+void UAttributeBasedRootMotionComponent::TickComponent(float DeltaTime, enum ELevelTick TickType, FActorComponentTickFunction* ThisTickFunction)
+{
 	static const FName RootMotionAttributeName = "RootMotionDelta";
 	static const UE::Anim::FAttributeId RootMotionAttributeId = { RootMotionAttributeName , FCompactPoseBoneIndex(0) };
 	
 	Super::TickComponent(DeltaTime, TickType, ThisTickFunction);
 	
-	ACharacter* Character = GetCharacterOwner();
-	check(Character);
+	if (ACharacter* Character = GetCharacterOwner())
+	{
+		if (const USkeletalMeshComponent* Mesh = Character->GetMesh())
+		{
+			if (const FTransformAnimationAttribute* RootMotionAttribute = Mesh->GetCustomAttributes().Find<FTransformAnimationAttribute>(RootMotionAttributeId))
+			{
+				TranslationVelocity = RootMotionAttribute->Value.GetTranslation() / DeltaTime;
+				RotationVelocity = RootMotionAttribute->Value.GetRotation().ToRotationVector() / DeltaTime;
 
-	UCharacterMovementComponent* CharacterMovement = Character->GetCharacterMovement<UCharacterMovementComponent>();
+				if (Mode == EAttributeBasedRootMotionMode::ApplyDelta)
+				{
+					if (UCharacterMovementComponent* CharacterMovement = Character->GetCharacterMovement<UCharacterMovementComponent>())
+					{
+						CharacterMovement->RootMotionParams.Set(RootMotionAttribute->Value);
+					}
+				}
+				
+				UE_VLOG_ARROW(this, "Root Motion", Display, Character->GetActorTransform().GetLocation(), Character->GetActorTransform().GetLocation() + Mesh->GetComponentTransform().TransformVector(TranslationVelocity) * 0.1, FColor::Green, TEXT(""));
+			}
+		}
+	}
+
+}
+
+
+void FAttributeBasedRootMotionComponentPrePhysicsTickFunction::ExecuteTick(float DeltaTime, ELevelTick TickType, ENamedThreads::Type CurrentThread, const FGraphEventRef& MyCompletionGraphEvent)
+{
+	FActorComponentTickFunction::ExecuteTickHelper(Target, /*bTickInEditor=*/ false, DeltaTime, TickType, [this](float DilatedTime)
+		{
+			Target->PrePhysicsTickComponent(DilatedTime);
+		});
+}
+
+FString FAttributeBasedRootMotionComponentPrePhysicsTickFunction::DiagnosticMessage()
+{
+	return Target->GetFullName() + TEXT("[UAttributeBasedRootMotionComponent::PrePhysicsTick]");
+}
 
-	if (const USkeletalMeshComponent* Mesh = Character->GetMesh())
+FName FAttributeBasedRootMotionComponentPrePhysicsTickFunction::DiagnosticContext(bool bDetailed)
+{
+	if (bDetailed)
+	{
+		return FName(*FString::Printf(TEXT("AttributeBasedRootMotionComponentPrePhysicsTick/%s"), *GetFullNameSafe(Target)));
+	}
+
+	return FName(TEXT("AttributeBasedRootMotionComponentPrePhysicsTick"));
+}
+
+void UAttributeBasedRootMotionComponent::PrePhysicsTickComponent(float DeltaTime)
+{
+	if (!bEnableRootMotion)
+	{
+		return;
+	}
+	
+	if (ACharacter* Character = GetCharacterOwner())
 	{
-		if (const FTransformAnimationAttribute* RootMotionAttribute = Mesh->GetCustomAttributes().Find<FTransformAnimationAttribute>(RootMotionAttributeId))
+		if (UCharacterMovementComponent* CharacterMovement = Character->GetCharacterMovement<UCharacterMovementComponent>())
 		{
-			CharacterMovement->RootMotionParams.Set(RootMotionAttribute->Value);
+			FTransform RootMotion;
+			RootMotion.SetTranslation(TranslationVelocity * DeltaTime);
+			RootMotion.SetRotation(FQuat::MakeFromRotationVector(RotationVelocity*DeltaTime));
+			CharacterMovement->RootMotionParams.Set(RootMotion);
 		}
 	}
 }
diff --git a/Engine/Plugins/Animation/MotionWarping/Source/MotionWarping/Public/AttributeBasedRootMotionComponent.h b/Engine/Plugins/Animation/MotionWarping/Source/MotionWarping/Public/AttributeBasedRootMotionComponent.h
--- a/Engine/Plugins/Animation/MotionWarping/Source/MotionWarping/Public/AttributeBasedRootMotionComponent.h
+++ b/Engine/Plugins/Animation/MotionWarping/Source/MotionWarping/Public/AttributeBasedRootMotionComponent.h
@@ -8,6 +8,48 @@
 
 class ACharacter;
 
+USTRUCT()
+struct FAttributeBasedRootMotionComponentPrePhysicsTickFunction : public FTickFunction
+{
+	GENERATED_USTRUCT_BODY()
+
+	/** CharacterMovementComponent that is the target of this tick **/
+	class UAttributeBasedRootMotionComponent* Target;
+
+	/**
+	 * Abstract function actually execute the tick.
+	 * @param DeltaTime - frame time to advance, in seconds
+	 * @param TickType - kind of tick for this frame
+	 * @param CurrentThread - thread we are executing on, useful to pass along as new tasks are created
+	 * @param MyCompletionGraphEvent - completion event for this task. Useful for holding the completion of this task until certain child tasks are complete.
+	 **/
+	virtual void ExecuteTick(float DeltaTime, enum ELevelTick TickType, ENamedThreads::Type CurrentThread, const FGraphEventRef& MyCompletionGraphEvent) override;
+
+	/** Abstract function to describe this tick. Used to print messages about illegal cycles in the dependency graph **/
+	virtual FString DiagnosticMessage() override;
+	
+	/** Function used to describe this tick for active tick reporting. **/
+	virtual FName DiagnosticContext(bool bDetailed) override;
+};
+template<>
+struct TStructOpsTypeTraits<FAttributeBasedRootMotionComponentPrePhysicsTickFunction> : public TStructOpsTypeTraitsBase2<FAttributeBasedRootMotionComponentPrePhysicsTickFunction>
+{
+	enum
+	{
+		WithCopy = false
+	};
+};
+
+
+UENUM()
+enum class EAttributeBasedRootMotionMode
+{
+	ApplyDelta,
+	ApplyVelocity,
+};
+
+
+
 UCLASS(ClassGroup = Movement, meta = (BlueprintSpawnableComponent))
 class MOTIONWARPING_API UAttributeBasedRootMotionComponent : public UActorComponent
 {
@@ -17,7 +59,9 @@ public:
 	UAttributeBasedRootMotionComponent(const FObjectInitializer& ObjectInitializer);
 
 	virtual void InitializeComponent() override;
+	virtual void RegisterComponentTickFunctions(bool bRegister) override;
 
+	void PrePhysicsTickComponent(float DeltaTime);
 	virtual void TickComponent(float DeltaTime, enum ELevelTick TickType, FActorComponentTickFunction* ThisTickFunction) override;
 	
 	/** Gets the character this component belongs to */
@@ -26,12 +70,23 @@ public:
 	UPROPERTY(BlueprintReadWrite, Category = "Runtime")
 	bool bEnableRootMotion = true;
 	
+	
+	UPROPERTY(EditAnywhere, Category = "Settings")
+	EAttributeBasedRootMotionMode Mode = EAttributeBasedRootMotionMode::ApplyDelta;
+	
 protected:
 
+	/** Pre-physics tick function for this character */
+	FAttributeBasedRootMotionComponentPrePhysicsTickFunction PrePhysicsTickFunction;
+
 	/** Character this component belongs to */
 	UPROPERTY(Transient)
 	TWeakObjectPtr<ACharacter> CharacterOwner;
 
 	FTransform PrevTransform;
 	bool PrevTransformValid = false;
-};
\ No newline at end of file
+
+	FVector TranslationVelocity {0, 0, 0};
+	FVector RotationVelocity {0, 0, 0};
+};
+
