fix AttributeBasedRootMotion Enable flag

diff --git a/Engine/Plugins/Animation/MotionWarping/Source/MotionWarping/Private/AttributeBasedRootMotionComponent.cpp b/Engine/Plugins/Animation/MotionWarping/Source/MotionWarping/Private/AttributeBasedRootMotionComponent.cpp
--- a/Engine/Plugins/Animation/MotionWarping/Source/MotionWarping/Private/AttributeBasedRootMotionComponent.cpp
+++ b/Engine/Plugins/Animation/MotionWarping/Source/MotionWarping/Private/AttributeBasedRootMotionComponent.cpp
@@ -66,33 +66,35 @@ void UAttributeBasedRootMotionComponent::RegisterComponentTickFunctions(bool bRe
 
 void UAttributeBasedRootMotionComponent::TickComponent(float DeltaTime, enum ELevelTick TickType, FActorComponentTickFunction* ThisTickFunction)
 {
-	static const FName RootMotionAttributeName = "RootMotionDelta";
-	static const UE::Anim::FAttributeId RootMotionAttributeId = { RootMotionAttributeName , FCompactPoseBoneIndex(0) };
+	if (bEnableRootMotion)
+	{
+		static const FName RootMotionAttributeName = "RootMotionDelta";
+		static const UE::Anim::FAttributeId RootMotionAttributeId = { RootMotionAttributeName , FCompactPoseBoneIndex(0) };
 	
-	Super::TickComponent(DeltaTime, TickType, ThisTickFunction);
+		Super::TickComponent(DeltaTime, TickType, ThisTickFunction);
 	
-	if (ACharacter* Character = GetCharacterOwner())
-	{
-		if (const USkeletalMeshComponent* Mesh = Character->GetMesh())
+		if (ACharacter* Character = GetCharacterOwner())
 		{
-			if (const FTransformAnimationAttribute* RootMotionAttribute = Mesh->GetCustomAttributes().Find<FTransformAnimationAttribute>(RootMotionAttributeId))
+			if (const USkeletalMeshComponent* Mesh = Character->GetMesh())
 			{
-				TranslationVelocity = RootMotionAttribute->Value.GetTranslation() / DeltaTime;
-				RotationVelocity = RootMotionAttribute->Value.GetRotation().ToRotationVector() / DeltaTime;
-
-				if (Mode == EAttributeBasedRootMotionMode::ApplyDelta)
+				if (const FTransformAnimationAttribute* RootMotionAttribute = Mesh->GetCustomAttributes().Find<FTransformAnimationAttribute>(RootMotionAttributeId))
 				{
-					if (UCharacterMovementComponent* CharacterMovement = Character->GetCharacterMovement<UCharacterMovementComponent>())
+					TranslationVelocity = RootMotionAttribute->Value.GetTranslation() / DeltaTime;
+					RotationVelocity = RootMotionAttribute->Value.GetRotation().ToRotationVector() / DeltaTime;
+
+					if (Mode == EAttributeBasedRootMotionMode::ApplyDelta)
 					{
-						CharacterMovement->RootMotionParams.Set(RootMotionAttribute->Value);
+						if (UCharacterMovementComponent* CharacterMovement = Character->GetCharacterMovement<UCharacterMovementComponent>())
+						{
+							CharacterMovement->RootMotionParams.Set(RootMotionAttribute->Value);
+						}
 					}
-				}
 				
-				UE_VLOG_ARROW(this, "Root Motion", Display, Character->GetActorTransform().GetLocation(), Character->GetActorTransform().GetLocation() + Mesh->GetComponentTransform().TransformVector(TranslationVelocity) * 0.1, FColor::Green, TEXT(""));
+					UE_VLOG_ARROW(this, "Root Motion", Display, Character->GetActorTransform().GetLocation(), Character->GetActorTransform().GetLocation() + Mesh->GetComponentTransform().TransformVector(TranslationVelocity) * 0.1, FColor::Green, TEXT(""));
+				}
 			}
 		}
 	}
-
 }
 
 
