Fix for crash when opening montages using warping windows

diff --git a/Engine/Plugins/Animation/MotionWarping/Source/MotionWarping/Private/AnimNotifyState_MotionWarping.cpp b/Engine/Plugins/Animation/MotionWarping/Source/MotionWarping/Private/AnimNotifyState_MotionWarping.cpp
--- a/Engine/Plugins/Animation/MotionWarping/Source/MotionWarping/Private/AnimNotifyState_MotionWarping.cpp
+++ b/Engine/Plugins/Animation/MotionWarping/Source/MotionWarping/Private/AnimNotifyState_MotionWarping.cpp
@@ -121,6 +121,10 @@ void UAnimNotifyState_MotionWarping::DrawInEditor(FPrimitiveDrawInterface* PDI,
 	{
 		FAnimNotifyEvent NotifyEventWithColor = NotifyEvent;
 		NotifyEventWithColor.NotifyColor = NotifyColor;
+		
+		// Necessary for FCompactPose (likely used by most RootMotionModifiers) that uses a FAnimStackAllocator (TMemStackAllocator) which allocates from FMemStack.
+		// When allocating memory from FMemStack we need to explicitly use FMemMark to ensure items are freed when the scope exits. 
+		FMemMark Mark(FMemStack::Get());
 		RootMotionModifier->DrawInEditor(PDI, MeshComp, Animation, NotifyEventWithColor);
 	}
 }
@@ -131,6 +135,10 @@ void UAnimNotifyState_MotionWarping::DrawCanvasInEditor(FCanvas& Canvas, FSceneV
 	{
 		FAnimNotifyEvent NotifyEventWithColor = NotifyEvent;
 		NotifyEventWithColor.NotifyColor = NotifyColor;
+
+		// Necessary for FCompactPose (likely used by most RootMotionModifiers) that uses a FAnimStackAllocator (TMemStackAllocator) which allocates from FMemStack.
+		// When allocating memory from FMemStack we need to explicitly use FMemMark to ensure items are freed when the scope exits.
+		FMemMark Mark(FMemStack::Get());
 		RootMotionModifier->DrawCanvasInEditor(Canvas, View, MeshComp, Animation, NotifyEventWithColor);
 	}
 }
