[MotionWarping] Add Opposite Default Motion Warp Rotation Type

diff --git a/Engine/Plugins/Animation/MotionWarping/Source/MotionWarping/Private/RootMotionModifier.cpp b/Engine/Plugins/Animation/MotionWarping/Source/MotionWarping/Private/RootMotionModifier.cpp
--- a/Engine/Plugins/Animation/MotionWarping/Source/MotionWarping/Private/RootMotionModifier.cpp
+++ b/Engine/Plugins/Animation/MotionWarping/Source/MotionWarping/Private/RootMotionModifier.cpp
@@ -301,18 +301,25 @@ void URootMotionModifier_Warp::OnTargetTransformChanged()
 
 FQuat URootMotionModifier_Warp::GetTargetRotation() const
 {
-	if (RotationType == EMotionWarpRotationType::Default)
+	switch (RotationType)
 	{
-		return CachedTargetTransform.GetRotation();
-	}
-	else if (RotationType == EMotionWarpRotationType::Facing)
-	{
-		if (const AActor* ActorOwner = GetActorOwner())
-		{
-			const FTransform& ActorTransform = ActorOwner->GetActorTransform();
-			const FVector ToSyncPoint = (CachedTargetTransform.GetLocation() - ActorTransform.GetLocation()).GetSafeNormal2D();
-			return FRotationMatrix::MakeFromXZ(ToSyncPoint, FVector::UpVector).ToQuat();
-		}
+		case EMotionWarpRotationType::Default:
+			return CachedTargetTransform.GetRotation();
+		
+		case EMotionWarpRotationType::Facing:
+			if (const AActor* ActorOwner = GetActorOwner())
+			{
+				const FTransform& ActorTransform = ActorOwner->GetActorTransform();
+				const FVector ToSyncPoint = (CachedTargetTransform.GetLocation() - ActorTransform.GetLocation()).GetSafeNormal2D();
+				return FRotationMatrix::MakeFromXZ(ToSyncPoint, FVector::UpVector).ToQuat();
+			}
+			break;
+		
+		case EMotionWarpRotationType::OppositeDefault:
+			return FRotationMatrix::MakeFromXZ(CachedTargetTransform.GetRotation().GetForwardVector() * -1, CachedTargetTransform.GetRotation().GetUpVector()).ToQuat();
+		
+		default:
+			checkNoEntry();
 	}
 
 	return FQuat::Identity;
diff --git a/Engine/Plugins/Animation/MotionWarping/Source/MotionWarping/Public/RootMotionModifier.h b/Engine/Plugins/Animation/MotionWarping/Source/MotionWarping/Public/RootMotionModifier.h
--- a/Engine/Plugins/Animation/MotionWarping/Source/MotionWarping/Public/RootMotionModifier.h
+++ b/Engine/Plugins/Animation/MotionWarping/Source/MotionWarping/Public/RootMotionModifier.h
@@ -240,6 +240,9 @@ enum class EMotionWarpRotationType : uint8
 
 	/** Character rotates to face the sync point */
 	Facing,
+
+	/** Character rotates to match the rotation of the sync point, rotated by half of full angle around up axis (180 degrees yaw) */
+	OppositeDefault,
 };
 
 UENUM(BlueprintType)
