[MotionWarping] Add additional way of debugging motion warping trargets - CVarWarpedTargetDebug

diff --git a/Engine/Plugins/Animation/MotionWarping/Source/MotionWarping/Private/MotionWarpingComponent.cpp b/Engine/Plugins/Animation/MotionWarping/Source/MotionWarping/Private/MotionWarpingComponent.cpp
--- a/Engine/Plugins/Animation/MotionWarping/Source/MotionWarping/Private/MotionWarpingComponent.cpp
+++ b/Engine/Plugins/Animation/MotionWarping/Source/MotionWarping/Private/MotionWarpingComponent.cpp
@@ -25,6 +25,7 @@ DEFINE_LOG_CATEGORY(LogMotionWarping);
 TAutoConsoleVariable<int32> FMotionWarpingCVars::CVarMotionWarpingDisable(TEXT("a.MotionWarping.Disable"), 0, TEXT("Disable Motion Warping"), ECVF_Cheat);
 TAutoConsoleVariable<int32> FMotionWarpingCVars::CVarMotionWarpingDebug(TEXT("a.MotionWarping.Debug"), 0, TEXT("0: Disable, 1: Only Log, 2: Only DrawDebug, 3: Log and DrawDebug"), ECVF_Cheat);
 TAutoConsoleVariable<float> FMotionWarpingCVars::CVarMotionWarpingDrawDebugDuration(TEXT("a.MotionWarping.DrawDebugLifeTime"), 1.f, TEXT("Time in seconds each draw debug persists.\nRequires 'a.MotionWarping.Debug 2'"), ECVF_Cheat);
+TAutoConsoleVariable<int32> FMotionWarpingCVars::CVarWarpedTargetDebug(TEXT("a.MotionWarping.Debug.Target"), false, TEXT("Shows warp target debug. 0 - disabled, 1 - enabled for selected actor, 2 - enabled for all actors"), ECVF_Cheat);
 #endif
 
 // UMotionWarpingUtilities
@@ -497,6 +498,53 @@ FTransform UMotionWarpingComponent::ProcessRootMotionPreConvertToWorld(const FTr
 
 		DrawDebugPoint(GetWorld(), ActorFeetLocation, PointSize, FColor::Blue, false, DrawDebugDuration, 0);
 	}
+
+	const int32 DebugValTarget = FMotionWarpingCVars::CVarWarpedTargetDebug->GetInt();
+	const bool bDebugTarget = (DebugValTarget == 1 && GetOwner()->IsSelected()) || DebugValTarget == 2;
+
+	const FVector ActorLocation = GetOwner()->GetActorLocation();
+	FVector TextLocation = ActorLocation;
+	constexpr float VerticalTextOffset = -10.0f;
+
+	TArray<URootMotionModifier_Warp*> WarpModifiers;
+	for (URootMotionModifier* Modifier : Modifiers)
+	{
+		if (Modifier->GetState() == ERootMotionModifierState::Active)
+		{
+			if (URootMotionModifier_Warp* WarpModifier = Cast<URootMotionModifier_Warp>(Modifier))
+			{
+				WarpModifiers.Add(WarpModifier);
+			}
+		}
+	}
+
+	for (int32 i = 0; i < WarpTargets.Num(); ++i)
+	{
+		const FMotionWarpingTarget& WarpingTarget = WarpTargets[i];
+
+		// Skip inactive warp targets
+		if (!WarpModifiers.ContainsByPredicate([&WarpingTarget](const URootMotionModifier_Warp* Mod) {return Mod->WarpTargetName == WarpingTarget.Name; }))
+		{
+			continue;
+		}
+
+		// Cycle between colors for better readability
+		FColor WarpTargetColor = (i % 2) != 0 ? FColor(21, 76, 190) : FColor::Green;
+
+		if (bDebugTarget)
+		{
+			FVector TargetLocation = WarpingTarget.GetTargetTrasform().GetLocation();
+			DrawDebugSphere(GetWorld(), TargetLocation, 5.0f, 16, WarpTargetColor, false);
+
+			DrawDebugLine(GetWorld(), TextLocation, WarpingTarget.GetTargetTrasform().GetLocation(), WarpTargetColor);
+
+			FString DebugText = FString::Printf(TEXT("Warp target name: %s, is dynamic: %s"), *WarpingTarget.Name.ToString(), WarpingTarget.bFollowComponent ? TEXT("True") : TEXT("False"));
+			DrawDebugString(GetWorld(), TextLocation, DebugText, nullptr, WarpTargetColor, 0.f, true);
+			TextLocation.Z += VerticalTextOffset;
+		}
+
+		TextLocation.Z += VerticalTextOffset;
+	}
 #endif
 
 	return FinalRootMotion;
diff --git a/Engine/Plugins/Animation/MotionWarping/Source/MotionWarping/Public/MotionWarpingComponent.h b/Engine/Plugins/Animation/MotionWarping/Source/MotionWarping/Public/MotionWarpingComponent.h
--- a/Engine/Plugins/Animation/MotionWarping/Source/MotionWarping/Public/MotionWarpingComponent.h
+++ b/Engine/Plugins/Animation/MotionWarping/Source/MotionWarping/Public/MotionWarpingComponent.h
@@ -25,6 +25,7 @@ struct MOTIONWARPING_API FMotionWarpingCVars
 	static TAutoConsoleVariable<int32> CVarMotionWarpingDisable;
 	static TAutoConsoleVariable<int32> CVarMotionWarpingDebug;
 	static TAutoConsoleVariable<float> CVarMotionWarpingDrawDebugDuration;
+	static TAutoConsoleVariable<int32> CVarWarpedTargetDebug;
 };
 #endif
 
