[MotionWarping] Add bWarpToFeetLocation, allowing warping using actor location, not feet location.

diff --git a/Engine/Plugins/Animation/MotionWarping/Source/MotionWarping/Private/RootMotionModifier_SkewWarp.cpp b/Engine/Plugins/Animation/MotionWarping/Source/MotionWarping/Private/RootMotionModifier_SkewWarp.cpp
--- a/Engine/Plugins/Animation/MotionWarping/Source/MotionWarping/Private/RootMotionModifier_SkewWarp.cpp
+++ b/Engine/Plugins/Animation/MotionWarping/Source/MotionWarping/Private/RootMotionModifier_SkewWarp.cpp
@@ -160,7 +160,12 @@ FTransform URootMotionModifier_SkewWarp::ProcessRootMotion(const FTransform& InR
 	{
 		const FVector CurrentLocation       = OwnerAdapter->GetVisualRootLocation();
 		const FQuat CurrentRotation         = OwnerAdapter->GetActor()->GetActorQuat();
-		const FVector MeshTranslationOffset = OwnerAdapter->GetBaseVisualTranslationOffset();
+		FVector MeshTranslationOffset       = OwnerAdapter->GetBaseVisualTranslationOffset();
+		if (!bWarpToFeetLocation)
+		{
+			MeshTranslationOffset.Z = 0.f;
+		}
+
 		const FQuat MeshRotationOffset      = OwnerAdapter->GetBaseVisualRotationOffset();
 
 		const FVector DeltaTranslation = RootMotionDelta.GetLocation();
diff --git a/Engine/Plugins/Animation/MotionWarping/Source/MotionWarping/Public/RootMotionModifier.h b/Engine/Plugins/Animation/MotionWarping/Source/MotionWarping/Public/RootMotionModifier.h
--- a/Engine/Plugins/Animation/MotionWarping/Source/MotionWarping/Public/RootMotionModifier.h
+++ b/Engine/Plugins/Animation/MotionWarping/Source/MotionWarping/Public/RootMotionModifier.h
@@ -300,6 +300,10 @@ public:
 	UPROPERTY(EditAnywhere, BlueprintReadWrite, Category = "Config", meta = (EditCondition = "bWarpTranslation"))
 	bool bIgnoreZAxis = true;
 
+	/** Whether we warp the actors location or their foot location to the warp target */
+	UPROPERTY(EditAnywhere, BlueprintReadWrite, Category = "Config", meta = (EditCondition = "bWarpTranslation"))
+	bool bWarpToFeetLocation = true;
+
 	/** Easing function used when adding translation. Only relevant when there is no translation in the animation */
 	UPROPERTY(EditAnywhere, BlueprintReadWrite, Category = "Config", meta = (EditCondition = "bWarpTranslation"))
 	EAlphaBlendOption AddTranslationEasingFunc = EAlphaBlendOption::Linear;
