RootMotion: Fix for root motion velocity rotation pop

diff --git a/Engine/Plugins/Animation/MotionWarping/Source/MotionWarping/Private/AttributeBasedRootMotionComponent.cpp b/Engine/Plugins/Animation/MotionWarping/Source/MotionWarping/Private/AttributeBasedRootMotionComponent.cpp
--- a/Engine/Plugins/Animation/MotionWarping/Source/MotionWarping/Private/AttributeBasedRootMotionComponent.cpp
+++ b/Engine/Plugins/Animation/MotionWarping/Source/MotionWarping/Private/AttributeBasedRootMotionComponent.cpp
@@ -80,7 +80,10 @@ void UAttributeBasedRootMotionComponent::TickComponent(float DeltaTime, enum ELe
 				if (const FTransformAnimationAttribute* RootMotionAttribute = Mesh->GetCustomAttributes().Find<FTransformAnimationAttribute>(RootMotionAttributeId))
 				{
 					TranslationVelocity = RootMotionAttribute->Value.GetTranslation() / DeltaTime;
-					RotationVelocity = RootMotionAttribute->Value.GetRotation().ToRotationVector() / DeltaTime;
+
+					FQuat RootMotionRotation = RootMotionAttribute->Value.GetRotation();
+					RootMotionRotation.EnforceShortestArcWith(FQuat::Identity);
+					RotationVelocity = RootMotionRotation.ToRotationVector() / DeltaTime;
 
 					if (Mode == EAttributeBasedRootMotionMode::ApplyDelta)
 					{
