[MotionWarping] Add option to choose location offset direction

diff --git a/Engine/Plugins/Animation/MotionWarping/Source/MotionWarping/Private/MotionWarpingComponent.cpp b/Engine/Plugins/Animation/MotionWarping/Source/MotionWarping/Private/MotionWarpingComponent.cpp
--- a/Engine/Plugins/Animation/MotionWarping/Source/MotionWarping/Private/MotionWarpingComponent.cpp
+++ b/Engine/Plugins/Animation/MotionWarping/Source/MotionWarping/Private/MotionWarpingComponent.cpp
@@ -622,7 +622,7 @@ void UMotionWarpingComponent::AddOrUpdateWarpTargetFromTransform(FName WarpTarge
 	AddOrUpdateWarpTarget(FMotionWarpingTarget(WarpTargetName, TargetTransform));
 }
 
-void UMotionWarpingComponent::AddOrUpdateWarpTargetFromComponent(FName WarpTargetName, const USceneComponent* Component, FName BoneName, bool bFollowComponent, FVector LocationOffset, FRotator RotationOffset)
+void UMotionWarpingComponent::AddOrUpdateWarpTargetFromComponent(FName WarpTargetName, const USceneComponent* Component, FName BoneName, bool bFollowComponent, EWarpTargetLocationOffsetDirection LocationOffsetDirection, FVector LocationOffset, FRotator RotationOffset)
 {
 	if (Component == nullptr)
 	{
@@ -630,7 +630,7 @@ void UMotionWarpingComponent::AddOrUpdateWarpTargetFromComponent(FName WarpTarge
 		return;
 	}
 
-	AddOrUpdateWarpTarget(FMotionWarpingTarget(WarpTargetName, Component, BoneName, bFollowComponent, LocationOffset, RotationOffset));
+	AddOrUpdateWarpTarget(FMotionWarpingTarget(WarpTargetName, Component, BoneName, bFollowComponent, LocationOffsetDirection, GetOwner(), LocationOffset, RotationOffset));
 }
 
 URootMotionModifier* UMotionWarpingComponent::AddModifierFromTemplate(URootMotionModifier* Template, const UAnimSequenceBase* Animation, float StartTime, float EndTime)
diff --git a/Engine/Plugins/Animation/MotionWarping/Source/MotionWarping/Private/RootMotionModifier.cpp b/Engine/Plugins/Animation/MotionWarping/Source/MotionWarping/Private/RootMotionModifier.cpp
--- a/Engine/Plugins/Animation/MotionWarping/Source/MotionWarping/Private/RootMotionModifier.cpp
+++ b/Engine/Plugins/Animation/MotionWarping/Source/MotionWarping/Private/RootMotionModifier.cpp
@@ -14,7 +14,7 @@
 // FMotionWarpingTarget
 ///////////////////////////////////////////////////////////////
 
-FMotionWarpingTarget::FMotionWarpingTarget(const FName& InName, const USceneComponent* InComp, FName InBoneName, bool bInFollowComponent, const FVector& InLocOffset, const FRotator& InRotOffset)
+FMotionWarpingTarget::FMotionWarpingTarget(const FName& InName, const USceneComponent* InComp, FName InBoneName, bool bInFollowComponent, EWarpTargetLocationOffsetDirection InLocationOffsetDirection, const AActor* InAvatarActor, const FVector& InLocOffset, const FRotator& InRotOffset)
 {
 	if (ensure(InComp))
 	{
@@ -24,6 +24,8 @@ FMotionWarpingTarget::FMotionWarpingTarget(const FName& InName, const USceneComp
 		bFollowComponent = bInFollowComponent;
 		LocationOffset = InLocOffset;
 		RotationOffset = InRotOffset;
+		AvatarActor = InAvatarActor;
+		LocationOffsetDirection = InLocationOffsetDirection;
 
 		FTransform Transform = FTransform::Identity;
 		if (BoneName != NAME_None)
@@ -35,13 +37,19 @@ FMotionWarpingTarget::FMotionWarpingTarget(const FName& InName, const USceneComp
 			Transform = InComp->GetComponentTransform();
 		}
 
-		Transform = FTransform(RotationOffset, LocationOffset) * Transform;
+		CacheOffset(Transform);
+		RecalculateOffset(Transform);
 
 		Location = Transform.GetLocation();
 		Rotation = Transform.Rotator();
 	}
 }
 
+FMotionWarpingTarget::FMotionWarpingTarget(const FName& InName, const USceneComponent* InComp, FName InBoneName, bool bInFollowComponent, const FVector& InLocOffset, const FRotator& InRotOffset)
+	: FMotionWarpingTarget(InName, InComp, InBoneName, bInFollowComponent, EWarpTargetLocationOffsetDirection::TargetsForwardVector, nullptr, InLocOffset, InRotOffset)
+{
+}
+
 FTransform FMotionWarpingTarget::GetTargetTransformFromComponent(const USceneComponent* Comp, const FName& BoneName)
 {
 	if (Comp == nullptr)
@@ -75,12 +83,79 @@ FTransform FMotionWarpingTarget::GetTargetTrasform() const
 			Transform = Component->GetComponentTransform();
 		}
 
-		return FTransform(RotationOffset, LocationOffset) * Transform;
+		RecalculateOffset(Transform);
+		return Transform;
 	}
 
 	return FTransform(Rotation, Location);
 }
 
+/*	Because vector from target to owner changes during warping, offset needs to be cached. */
+void FMotionWarpingTarget::CacheOffset(const FTransform& InTransform)
+{
+	// Forward offset doesn't need to be cached if it's the only one used. Otherwise, cache it too.
+	bCacheForwardOffset =
+		LocationOffsetDirection == EWarpTargetLocationOffsetDirection::VectorFromTargetToOwner &&
+			LocationOffset.X > SMALL_NUMBER &&
+				(LocationOffset.Y > SMALL_NUMBER || LocationOffset.Z > SMALL_NUMBER);
+	
+	if (LocationOffsetDirection == EWarpTargetLocationOffsetDirection::VectorFromTargetToOwner)
+	{
+		const FVector ContextVector = (AvatarActor->GetActorLocation() - InTransform.GetLocation()).GetSafeNormal();
+		const FVector RightVector = -FVector::CrossProduct(ContextVector, FVector::UpVector);
+
+		if (bCacheForwardOffset)
+		{
+			CachedForwardOffset = (ContextVector * LocationOffset.X);
+		}
+		
+		CachedRightOffset = RightVector * LocationOffset.Y;
+		CachedUpOffset = -FVector::CrossProduct(RightVector, ContextVector) * LocationOffset.Z;
+	}
+}
+
+void FMotionWarpingTarget::RecalculateOffset(FTransform& InTransform) const
+{
+	FVector Offset = FVector::ZeroVector;
+	
+	switch (LocationOffsetDirection)
+	{
+		case EWarpTargetLocationOffsetDirection::TargetsForwardVector:
+			Offset = LocationOffset;
+			break;
+			
+		case EWarpTargetLocationOffsetDirection::VectorFromTargetToOwner:
+			if (AvatarActor.IsValid())
+			{
+				if (bCacheForwardOffset)
+				{
+					Offset = Component->GetComponentTransform().Inverse().TransformVector(CachedForwardOffset + CachedUpOffset + CachedRightOffset);
+				}
+				else
+				{
+					const FVector ContextVector = (AvatarActor->GetActorLocation() - InTransform.GetLocation()).GetSafeNormal();
+					const FVector ForwardOffset = (ContextVector * LocationOffset.X);
+					
+					Offset = Component->GetComponentTransform().Inverse().TransformVector(ForwardOffset + CachedUpOffset + CachedRightOffset);
+				}
+			}
+			else
+			{
+				UE_LOG(LogMotionWarping, Warning, TEXT("Motion warping offset is set to VectorFromOwnerToTarget but avator actor is invalid"))
+			}
+			break;
+			
+		case EWarpTargetLocationOffsetDirection::WorldSpace:
+			Offset = Component->GetComponentTransform().Inverse().TransformVector(LocationOffset);
+			break;
+			
+		default:
+			checkNoEntry();
+	}
+	
+	InTransform = FTransform(RotationOffset, Offset) * InTransform;
+}
+
 // URootMotionModifier
 ///////////////////////////////////////////////////////////////
 
diff --git a/Engine/Plugins/Animation/MotionWarping/Source/MotionWarping/Public/MotionWarpingComponent.h b/Engine/Plugins/Animation/MotionWarping/Source/MotionWarping/Public/MotionWarpingComponent.h
--- a/Engine/Plugins/Animation/MotionWarping/Source/MotionWarping/Public/MotionWarpingComponent.h
+++ b/Engine/Plugins/Animation/MotionWarping/Source/MotionWarping/Public/MotionWarpingComponent.h
@@ -162,10 +162,14 @@ public:
 	 * @param BoneName Optional bone or socket in the component used to get the target transform. 
 	 * @param bFollowComponent Whether the target transform should update while the warping is active. Useful for tracking moving targets.
 	 *		  Note that this will be one frame off if our owner ticks before the target actor. Add tick pre-requisites to avoid this.
+	 * @param LocationOffsetDirection Direction of LocationOffset
 	 * @param LocationOffset Optional translation offset to apply to the transform we get from the component
 	 * @param RotationOffset Optional rotation offset to apply to the transform we get from the component
 	 */
 	UFUNCTION(BlueprintCallable, Category = "Motion Warping")
+	void AddOrUpdateWarpTargetFromComponent(FName WarpTargetName, const USceneComponent* Component, FName BoneName, bool bFollowComponent, EWarpTargetLocationOffsetDirection LocationOffsetDirection = EWarpTargetLocationOffsetDirection::TargetsForwardVector, FVector LocationOffset = FVector::ZeroVector, FRotator RotationOffset = FRotator::ZeroRotator);
+	
+	// @TODO: Deprecate this.
 	void AddOrUpdateWarpTargetFromComponent(FName WarpTargetName, const USceneComponent* Component, FName BoneName, bool bFollowComponent, FVector LocationOffset = FVector::ZeroVector, FRotator RotationOffset = FRotator::ZeroRotator);
 
 	/**
diff --git a/Engine/Plugins/Animation/MotionWarping/Source/MotionWarping/Public/RootMotionModifier.h b/Engine/Plugins/Animation/MotionWarping/Source/MotionWarping/Public/RootMotionModifier.h
--- a/Engine/Plugins/Animation/MotionWarping/Source/MotionWarping/Public/RootMotionModifier.h
+++ b/Engine/Plugins/Animation/MotionWarping/Source/MotionWarping/Public/RootMotionModifier.h
@@ -161,6 +161,17 @@ private:
 	ERootMotionModifierState State = ERootMotionModifierState::Waiting;
 };
 
+UENUM(BlueprintType)
+enum class EWarpTargetLocationOffsetDirection : uint8
+{
+	/* Location Offset's X is target component's forward vector. */
+	TargetsForwardVector,
+	/* Location Offset's X is vector from owner to target. */
+	VectorFromTargetToOwner,
+	/* Location Offset's X is world space X. */
+	WorldSpace
+};
+
 // URootMotionModifier_Warp
 ///////////////////////////////////////////////////////////////
 
@@ -194,6 +205,10 @@ struct MOTIONWARPING_API FMotionWarpingTarget
 	UPROPERTY(EditAnywhere, BlueprintReadWrite, Category = "Defaults")
 	bool bFollowComponent;
 
+	/** Direction of offset context vector */
+	UPROPERTY(EditAnywhere, BlueprintReadWrite, Category = "Defaults")
+	EWarpTargetLocationOffsetDirection LocationOffsetDirection = EWarpTargetLocationOffsetDirection::TargetsForwardVector;
+
 	/** Optional static location offset. Only relevant when the warp target is created from a component */
 	UPROPERTY(EditAnywhere, BlueprintReadWrite, Category = "Defaults")
 	FVector LocationOffset = FVector::ZeroVector;
@@ -202,15 +217,29 @@ struct MOTIONWARPING_API FMotionWarpingTarget
 	UPROPERTY(EditAnywhere, BlueprintReadWrite, Category = "Defaults")
 	FRotator RotationOffset = FRotator::ZeroRotator;
 
+	/** Optional Actor that uses this warping target. Used for offset calculation of target transform is provided via component  */
+	UPROPERTY(EditAnywhere, BlueprintReadWrite, Category = "Defaults")
+	TWeakObjectPtr<const AActor> AvatarActor;
+
+	bool bCacheForwardOffset = false;
+	FVector CachedForwardOffset = FVector::ZeroVector;
+	FVector CachedRightOffset = FVector::ZeroVector;
+	FVector CachedUpOffset = FVector::ZeroVector;
+
 	FMotionWarpingTarget()
 		: Name(NAME_None), Location(FVector::ZeroVector), Rotation(FRotator::ZeroRotator), Component(nullptr), BoneName(NAME_None), bFollowComponent(false) {}
 
 	FMotionWarpingTarget(const FName& InName, const FTransform& InTransform)
 		: Name(InName), Location(InTransform.GetLocation()), Rotation(InTransform.Rotator()), Component(nullptr), BoneName(NAME_None), bFollowComponent(false) {}
 
+	FMotionWarpingTarget(const FName& InName, const USceneComponent* InComp, FName InBoneName, bool bInbFollowComponent, EWarpTargetLocationOffsetDirection InLocationOffsetDirection = EWarpTargetLocationOffsetDirection::TargetsForwardVector, const AActor* InAvatarActor = nullptr, const FVector& InLocOffset = FVector::ZeroVector, const FRotator& InRotOffset = FRotator::ZeroRotator);
+	
+	// @TODO: Deprecate this.
 	FMotionWarpingTarget(const FName& InName, const USceneComponent* InComp, FName InBoneName, bool bInbFollowComponent, const FVector& InLocOffset = FVector::ZeroVector, const FRotator& InRotOffset = FRotator::ZeroRotator);
 
 	FTransform GetTargetTrasform() const;
+	void CacheOffset(const FTransform& InTransform);
+	void RecalculateOffset(FTransform& InTransform) const;
 
 	FORCEINLINE FVector GetLocation() const { return GetTargetTrasform().GetLocation(); }
 	FORCEINLINE FQuat GetRotation() const { return GetTargetTrasform().GetRotation(); }
@@ -219,13 +248,13 @@ struct MOTIONWARPING_API FMotionWarpingTarget
 	FORCEINLINE bool operator==(const FMotionWarpingTarget& Other) const
 	{
 		return Other.Name == Name && Other.Location.Equals(Location) && Other.Rotation.Equals(Rotation) && Other.Component == Component && Other.BoneName == BoneName && Other.bFollowComponent == bFollowComponent && 
-			Other.LocationOffset.Equals(LocationOffset) && Other.RotationOffset.Equals(RotationOffset);
+			Other.LocationOffsetDirection == LocationOffsetDirection && Other.LocationOffset.Equals(LocationOffset) && Other.RotationOffset.Equals(RotationOffset) && Other.AvatarActor.Get() == AvatarActor.Get();
 	}
 
 	FORCEINLINE bool operator!=(const FMotionWarpingTarget& Other) const
 	{
 		return Other.Name != Name || !Other.Location.Equals(Location) || !Other.Rotation.Equals(Rotation) || Other.Component != Component || Other.BoneName != BoneName || Other.bFollowComponent != bFollowComponent || 
-			!Other.LocationOffset.Equals(LocationOffset) || !Other.RotationOffset.Equals(RotationOffset);
+			Other.LocationOffsetDirection != LocationOffsetDirection || !Other.LocationOffset.Equals(LocationOffset) || !Other.RotationOffset.Equals(RotationOffset) || Other.AvatarActor.Get() != AvatarActor.Get();
 	}
 
 	static FTransform GetTargetTransformFromComponent(const USceneComponent* Comp, const FName& BoneName);
