[MotionWarping] Add option to subtract root motion remaining after motion warping window from warp target transform.

diff --git a/Engine/Plugins/Animation/MotionWarping/Source/MotionWarping/Private/RootMotionModifier.cpp b/Engine/Plugins/Animation/MotionWarping/Source/MotionWarping/Private/RootMotionModifier.cpp
--- a/Engine/Plugins/Animation/MotionWarping/Source/MotionWarping/Private/RootMotionModifier.cpp
+++ b/Engine/Plugins/Animation/MotionWarping/Source/MotionWarping/Private/RootMotionModifier.cpp
@@ -352,7 +352,9 @@ void URootMotionModifier_Warp::Update(const FMotionWarpingUpdateContext& Context
 			// Update Target Transform based on the offset between the root and the warp point in the animation
 			TargetTransform = CachedOffsetFromWarpPoint.GetValue() * WarpPointTransformGame;
 		}
-
+		
+		CachedTargetTransform *= RootMotionRemainingAfterNotify.Inverse();
+		
 		if (!CachedTargetTransform.Equals(TargetTransform))
 		{
 			CachedTargetTransform = TargetTransform;
@@ -374,6 +376,16 @@ void URootMotionModifier_Warp::OnTargetTransformChanged()
 	}
 }
 
+void URootMotionModifier_Warp::OnStateChanged(ERootMotionModifierState LastState)
+{
+	Super::OnStateChanged(LastState);
+
+	if (bSubtractRemainingRootMotion)
+	{
+		RootMotionRemainingAfterNotify = UMotionWarpingUtilities::ExtractRootMotionFromAnimation(Animation.Get(), EndTime, Animation.Get()->GetPlayLength());
+	}
+}
+
 FQuat URootMotionModifier_Warp::GetTargetRotation() const
 {
 	switch (RotationType)
@@ -408,7 +420,7 @@ FQuat URootMotionModifier_Warp::WarpRotation(const FTransform& RootMotionDelta,
 	if (const UMotionWarpingBaseAdapter* WarpingAdapter = GetOwnerAdapter())
 	{
 		CurrentRotation = WarpingAdapter->GetActor()->GetActorQuat() * WarpingAdapter->GetBaseVisualRotationOffset();
-		TargetRotation = CurrentRotation.Inverse() * (GetTargetRotation() * WarpingAdapter->GetBaseVisualRotationOffset());
+		TargetRotation = CurrentRotation.Inverse() * (GetTargetRotation() * WarpingAdapter->GetBaseVisualRotationOffset() * RootMotionRemainingAfterNotify.GetRotation().Inverse());
 	}
 	else
 	{
diff --git a/Engine/Plugins/Animation/MotionWarping/Source/MotionWarping/Private/RootMotionModifier_SkewWarp.cpp b/Engine/Plugins/Animation/MotionWarping/Source/MotionWarping/Private/RootMotionModifier_SkewWarp.cpp
--- a/Engine/Plugins/Animation/MotionWarping/Source/MotionWarping/Private/RootMotionModifier_SkewWarp.cpp
+++ b/Engine/Plugins/Animation/MotionWarping/Source/MotionWarping/Private/RootMotionModifier_SkewWarp.cpp
@@ -183,7 +183,7 @@ FTransform URootMotionModifier_SkewWarp::ProcessRootMotion(const FTransform& InR
 			if (!DeltaTranslation.IsNearlyZero())
 			{
 				const FTransform MeshTransform = FTransform(MeshRotationOffset, MeshTranslationOffset) * OwnerAsActor->GetActorTransform();
-				TargetLocation = MeshTransform.InverseTransformPositionNoScale(TargetLocation);
+				TargetLocation = MeshTransform.InverseTransformPositionNoScale(TargetLocation) - RootMotionRemainingAfterNotify.GetTranslation();
 
 				const FVector WarpedTranslation = WarpTranslation(FTransform::Identity, DeltaTranslation, TotalTranslation, TargetLocation) + ExtraRootMotion.GetLocation();
 				FinalRootMotion.SetTranslation(WarpedTranslation);
diff --git a/Engine/Plugins/Animation/MotionWarping/Source/MotionWarping/Public/RootMotionModifier.h b/Engine/Plugins/Animation/MotionWarping/Source/MotionWarping/Public/RootMotionModifier.h
--- a/Engine/Plugins/Animation/MotionWarping/Source/MotionWarping/Public/RootMotionModifier.h
+++ b/Engine/Plugins/Animation/MotionWarping/Source/MotionWarping/Public/RootMotionModifier.h
@@ -353,6 +353,11 @@ public:
 	UPROPERTY(EditAnywhere, BlueprintReadWrite, Category = "Config", meta = (EditCondition = "bWarpRotation"))
 	EMotionWarpRotationMethod RotationMethod = EMotionWarpRotationMethod::Slerp;
 
+	/** Should root motion after this notify state be subtracted from warp target transform. 
+	This means actor will arrive at warp target transform when the animation finishes, not when the notify ends */
+	UPROPERTY(EditAnywhere, BlueprintReadWrite, Category = "Config")
+	bool bSubtractRemainingRootMotion = false;
+
 	/**
 	 * Allow to modify how fast the rotation is warped.
 	 * e.g if the window duration is 2sec and this is 0.5, the target rotation will be reached in 1sec instead of 2sec
@@ -373,6 +378,8 @@ public:
 	/** Event called during update if the target transform changes while the warping is active */
 	virtual void OnTargetTransformChanged();
 
+	virtual void OnStateChanged(ERootMotionModifierState LastState) override;
+
 #if !(UE_BUILD_SHIPPING || UE_BUILD_TEST)
 	void PrintLog(const FString& Name, const FTransform& OriginalRootMotion, const FTransform& WarpedRootMotion) const;
 #endif
@@ -388,6 +395,9 @@ protected:
 	UPROPERTY()
 	FTransform CachedTargetTransform = FTransform::Identity;
 
+	UPROPERTY()
+	FTransform RootMotionRemainingAfterNotify = FTransform::Identity;
+
 	/** Cached of the offset from the warp target. Used to calculate the final target transform when a warp target is defined in the animation */
 	TOptional<FTransform> CachedOffsetFromWarpPoint;
 };
