[MotionWarping] Add switch off conditions

diff --git a/Engine/Plugins/Animation/MotionWarping/Source/MotionWarping/Private/MotionWarpingComponent.cpp b/Engine/Plugins/Animation/MotionWarping/Source/MotionWarping/Private/MotionWarpingComponent.cpp
--- a/Engine/Plugins/Animation/MotionWarping/Source/MotionWarping/Private/MotionWarpingComponent.cpp
+++ b/Engine/Plugins/Animation/MotionWarping/Source/MotionWarping/Private/MotionWarpingComponent.cpp
@@ -14,6 +14,7 @@
 #include "GameFramework/CharacterMovementComponent.h"
 #include "AnimNotifyState_MotionWarping.h"
 #include "MotionWarpingCharacterAdapter.h"
+#include "MotionWarpingSwitchOffCondition.h"
 #include "Net/Core/PushModel/PushModel.h"
 #include "Net/UnrealNetwork.h"
 
@@ -26,6 +27,7 @@ TAutoConsoleVariable<int32> FMotionWarpingCVars::CVarMotionWarpingDisable(TEXT("
 TAutoConsoleVariable<int32> FMotionWarpingCVars::CVarMotionWarpingDebug(TEXT("a.MotionWarping.Debug"), 0, TEXT("0: Disable, 1: Only Log, 2: Only DrawDebug, 3: Log and DrawDebug"), ECVF_Cheat);
 TAutoConsoleVariable<float> FMotionWarpingCVars::CVarMotionWarpingDrawDebugDuration(TEXT("a.MotionWarping.DrawDebugLifeTime"), 1.f, TEXT("Time in seconds each draw debug persists.\nRequires 'a.MotionWarping.Debug 2'"), ECVF_Cheat);
 TAutoConsoleVariable<int32> FMotionWarpingCVars::CVarWarpedTargetDebug(TEXT("a.MotionWarping.Debug.Target"), false, TEXT("Shows warp target debug. 0 - disabled, 1 - enabled for selected actor, 2 - enabled for all actors"), ECVF_Cheat);
+TAutoConsoleVariable<int32> FMotionWarpingCVars::CVarWarpedSwitchOffConditionDebug(TEXT("a.MotionWarping.Debug.SwitchOffCondition"), false, TEXT("Shows switch off condition debug. 0 - disabled, 1 - enabled for selected actor, 2 - enabled for all actors"), ECVF_Cheat);
 #endif
 
 // UMotionWarpingUtilities
@@ -335,8 +337,85 @@ void UMotionWarpingComponent::DisableAllRootMotionModifiers()
 	}
 }
 
+void UMotionWarpingComponent::UpdateSwitchOffConditions()
+{
+	for (int32 i = WarpTargets.Num() - 1; i >= 0; --i)
+	{
+		FSwitchOffConditionData* SwitchOffConditionData = FindSwitchOffConditionData(WarpTargets[i].Name);
+		if (!SwitchOffConditionData)
+		{
+			continue;
+		}
+
+		TArray<TObjectPtr<UMotionWarpingSwitchOffCondition>>* Conditions = &SwitchOffConditionData->SwitchOffConditions;
+		if (!Conditions)
+		{
+			continue;
+		}
+
+		bool bClearCondition = false;
+		bool bPauseWarping = false;
+		bool bPauseRootMotion = false;
+
+		for (const UMotionWarpingSwitchOffCondition* Condition : *Conditions)
+		{
+			if (!IsValid(Condition) || !Condition->IsConditionValid())
+			{
+				continue;
+			}
+			if (Condition->Check())
+			{
+				switch (Condition->GetEffect())
+				{
+				case ESwitchOffConditionEffect::CancelFollow:
+					if (WarpTargets[i].bFollowComponent)
+					{
+						WarpTargets[i].bFollowComponent = false;
+
+						if (const USceneComponent* Component = WarpTargets[i].Component.Get())
+						{
+							WarpTargets[i].Location = Component->GetComponentLocation();
+							WarpTargets[i].Rotation = Component->GetComponentRotation();
+						}
+						else
+						{
+							UE_LOG(LogMotionWarping, Warning, TEXT("Switch off condition on warp target %s failed. Trying to cancel follow component, but component isn't valid."), *WarpTargets[i].Name.ToString());
+						}
+					}
+					break;
+				case ESwitchOffConditionEffect::CancelWarping:
+					bClearCondition = true;
+					break;
+				case ESwitchOffConditionEffect::PauseWarping:
+					bPauseWarping = true;
+					break;
+				case ESwitchOffConditionEffect::PauseRootMotion:
+					bPauseRootMotion = true;
+					break;
+				default:
+					checkNoEntry();
+				}
+			}
+		}
+
+		//remove finished and invalid conditions
+		if (bClearCondition)
+		{
+			RemoveSwitchOffConditions(WarpTargets[i].Name);
+			WarpTargets.RemoveAtSwap(i);
+		}
+		else
+		{
+			WarpTargets[i].bWarpingPaused = bPauseWarping;
+			WarpTargets[i].bRootMotionPaused = bPauseRootMotion;
+		}
+	}
+}
+
 void UMotionWarpingComponent::UpdateWithContext(const FMotionWarpingUpdateContext& Context, float DeltaSeconds)
 {
+	UpdateSwitchOffConditions();
+
 	if (Context.Animation.IsValid())
 	{
 		const UAnimSequenceBase* Animation = Context.Animation.Get();
@@ -499,6 +578,10 @@ FTransform UMotionWarpingComponent::ProcessRootMotionPreConvertToWorld(const FTr
 		DrawDebugPoint(GetWorld(), ActorFeetLocation, PointSize, FColor::Blue, false, DrawDebugDuration, 0);
 	}
 
+	const int32 DebugValSwitchOffCondition = FMotionWarpingCVars::CVarWarpedSwitchOffConditionDebug->GetInt();
+
+	const bool bDebugSwitchOffCondition = (DebugValSwitchOffCondition == 1 && GetOwner()->IsSelected()) || DebugValSwitchOffCondition == 2;
+
 	const int32 DebugValTarget = FMotionWarpingCVars::CVarWarpedTargetDebug->GetInt();
 	const bool bDebugTarget = (DebugValTarget == 1 && GetOwner()->IsSelected()) || DebugValTarget == 2;
 
@@ -543,6 +626,29 @@ FTransform UMotionWarpingComponent::ProcessRootMotionPreConvertToWorld(const FTr
 			TextLocation.Z += VerticalTextOffset;
 		}
 
+		if (bDebugSwitchOffCondition)
+		{
+			if (FSwitchOffConditionData* ConditionData = FindSwitchOffConditionData(WarpingTarget.Name))
+			{
+				DrawDebugString(GetWorld(), TextLocation, TEXT("Switch off conditions:"), nullptr, FColor::White, 0.f, true);
+				TextLocation.Z += VerticalTextOffset;
+
+				for (const UMotionWarpingSwitchOffCondition* Condition : ConditionData->SwitchOffConditions)
+				{
+					const bool bConditionTrue = Condition->Check();
+					FColor SwitchOffConditionTextColor = bConditionTrue ? FColor::Red : FColor::Yellow;
+					DrawDebugString(GetWorld(),
+						TextLocation,
+						FString::Printf(TEXT("%s - %s"), *Condition->ExtraDebugInfo(), bConditionTrue ? TEXT("TRUE") : TEXT("FALSE")),
+						nullptr,
+						SwitchOffConditionTextColor,
+						0.f,
+						true);
+					TextLocation.Z += VerticalTextOffset;
+				}
+			}
+		}
+
 		TextLocation.Z += VerticalTextOffset;
 	}
 #endif
@@ -572,7 +678,12 @@ void UMotionWarpingComponent::AddOrUpdateWarpTarget(const FMotionWarpingTarget&
 		// if we did not find the target, add it
 		if (!FindAndUpdateWarpTarget(WarpTarget))
 		{
-			WarpTargets.Add(WarpTarget);
+			int32 Idx = WarpTargets.Add(WarpTarget);
+
+			if (FSwitchOffConditionData* SwitchOffConditionData = FindSwitchOffConditionData(WarpTarget.Name))
+			{
+				SwitchOffConditionData->SetMotionWarpingTarget(&WarpTargets[Idx]);
+			}
 		}
 
 		MARK_PROPERTY_DIRTY_FROM_NAME(UMotionWarpingComponent, WarpTargets, this);
@@ -588,12 +699,18 @@ int32 UMotionWarpingComponent::RemoveWarpTarget(FName WarpTargetName)
 		MARK_PROPERTY_DIRTY_FROM_NAME(UMotionWarpingComponent, WarpTargets, this);
 	}
 
+	RemoveSwitchOffConditions(WarpTargetName);
+
 	return NumRemoved;
 }
 
 int32 UMotionWarpingComponent::RemoveWarpTargets(const TArray<FName>& WarpTargetNames)
 {
-	const int32 NumRemoved = WarpTargets.RemoveAll([&WarpTargetNames](const FMotionWarpingTarget& WarpTarget) { return WarpTargetNames.Contains(WarpTarget.Name); });
+	int32 NumRemoved = 0;
+	for (const FName& WarpTargetName : WarpTargetNames)
+	{
+		NumRemoved += RemoveWarpTarget(WarpTargetName);		
+	}
 
 	if (NumRemoved > 0)
 	{
@@ -603,10 +720,60 @@ int32 UMotionWarpingComponent::RemoveWarpTargets(const TArray<FName>& WarpTarget
 	return NumRemoved;
 }
 
+void UMotionWarpingComponent::AddSwitchOffCondition(FName WarpTargetName, UMotionWarpingSwitchOffCondition* Condition)
+{
+	if (IsValid(Condition))
+	{
+		if (const FMotionWarpingTarget* MotionWarpingTarget = FindWarpTarget(WarpTargetName))
+		{
+			Condition->SetMotionWarpingTarget(MotionWarpingTarget);
+		}
+		
+		if (FSwitchOffConditionData* SwitchOffConditionData = FindSwitchOffConditionData(WarpTargetName))
+		{
+			SwitchOffConditionData->SwitchOffConditions.Add(Condition);
+		}
+		else
+		{
+			SwitchOffConditions.Add(FSwitchOffConditionData(WarpTargetName, Condition));
+		}
+	}
+	else
+	{
+		UE_LOG(LogMotionWarping, Error, TEXT("Trying to add invalid switch off condition"))
+	}
+}
+
+void UMotionWarpingComponent::RemoveSwitchOffConditions(FName WarpTargetName)
+{
+	const int32 Index = SwitchOffConditions.IndexOfByPredicate([WarpTargetName](FSwitchOffConditionData Condition)
+	{
+		return Condition.WarpTargetName == WarpTargetName;
+	});
+
+	if (Index != INDEX_NONE)
+	{
+		SwitchOffConditions.RemoveAtSwap(Index);
+	}	
+}
+
+FSwitchOffConditionData* UMotionWarpingComponent::FindSwitchOffConditionData(FName WarpTargetName)
+{
+	return SwitchOffConditions.FindByPredicate([WarpTargetName](FSwitchOffConditionData Condition)
+	{
+		return Condition.WarpTargetName == WarpTargetName;
+	});
+}
+
 int32 UMotionWarpingComponent::RemoveAllWarpTargets()
 {
 	const int32 NumRemoved = WarpTargets.Num();
 
+	for (const FMotionWarpingTarget& Target : WarpTargets)
+	{
+		RemoveSwitchOffConditions(Target.Name);
+	}
+	
 	WarpTargets.Reset();
 
 	if (NumRemoved > 0)
diff --git a/Engine/Plugins/Animation/MotionWarping/Source/MotionWarping/Private/MotionWarpingSwitchOffCondition.cpp b/Engine/Plugins/Animation/MotionWarping/Source/MotionWarping/Private/MotionWarpingSwitchOffCondition.cpp
new file mode 100644
--- /dev/null
+++ b/Engine/Plugins/Animation/MotionWarping/Source/MotionWarping/Private/MotionWarpingSwitchOffCondition.cpp
@@ -0,0 +1,316 @@
+// Copyright Epic Games, Inc. All Rights Reserved.
+
+#include "MotionWarping/Public/MotionWarpingSwitchOffCondition.h"
+#include "MotionWarpingComponent.h"
+#include "GameFramework/Actor.h"
+#include "UObject/Package.h"
+
+void FSwitchOffConditionData::SetMotionWarpingTarget(const FMotionWarpingTarget* MotionWarpingTarget)
+{
+	for (UMotionWarpingSwitchOffCondition* SwitchOffCondition : SwitchOffConditions)
+	{
+		SwitchOffCondition->SetMotionWarpingTarget(MotionWarpingTarget);
+	}
+}
+
+bool UMotionWarpingSwitchOffCondition::Check() const
+{
+	return OnCheck();
+}
+
+FVector UMotionWarpingSwitchOffCondition::GetTargetLocation() const
+{
+	if (bUseWarpTargetAsTargetLocation)
+	{
+		if (MotionWarpingTarget)
+		{
+			return MotionWarpingTarget->GetLocation();
+		}
+
+		UE_LOG(LogMotionWarping, Warning, TEXT("Switch off condition is set to use warp target as target location, "
+			"however warp target appears to be null. Make sure warp target is correctly passed, "
+			"otherwise switch off condition will use target actor location."))
+	}
+
+	return TargetActor->GetActorLocation();
+}
+
+void UMotionWarpingSwitchOffCondition::SetWarpTargetForDestination(const FMotionWarpingTarget* InMotionWarpingTarget)
+{
+	MotionWarpingTarget = InMotionWarpingTarget;
+}
+
+bool UMotionWarpingSwitchOffCondition::IsConditionValid() const
+{
+	if (!IsValid(OwnerActor))
+	{
+		
+		UE_LOG(LogMotionWarping, Display, TEXT("%s won't work due to invalid Owner Actor"), *GetClass()->GetName())
+		return false;
+	}
+	
+	if (bUseWarpTargetAsTargetLocation)
+	{
+		if (!MotionWarpingTarget)
+		{
+
+			UE_LOG(LogMotionWarping, Display, TEXT("%s is set to use Motion Warping Target as target location, but won't work due to null Motion Warping Target"), *GetClass()->GetName())
+			return false;
+		}
+
+		return true;
+	}
+
+	if (!IsValid(TargetActor))
+	{
+		UE_LOG(LogMotionWarping, Display, TEXT("%s on actor %s is set to use Actor as target location, but won't work due to invalid Target Actor"), *GetClass()->GetName(), *OwnerActor->GetName());
+		return  false;
+	}
+	
+	return true;
+}
+
+UMotionWarpingSwitchOffDistanceCondition* UMotionWarpingSwitchOffDistanceCondition::CreateSwitchOffDistanceCondition(AActor* InOwnerActor, ESwitchOffConditionEffect InEffect, ESwitchOffConditionDistanceOp InOperator, float InDistance, bool InbUseWarpTargetAsTargetLocation, AActor* InTargetActor)
+{
+	UMotionWarpingSwitchOffDistanceCondition* SwitchOffDistanceCondition = NewObject<UMotionWarpingSwitchOffDistanceCondition>();
+	SwitchOffDistanceCondition->OwnerActor = InOwnerActor;
+	SwitchOffDistanceCondition->Effect = InEffect;
+	SwitchOffDistanceCondition->Operator = InOperator;
+	SwitchOffDistanceCondition->Distance = InDistance;
+	SwitchOffDistanceCondition->bUseWarpTargetAsTargetLocation = InbUseWarpTargetAsTargetLocation;
+	SwitchOffDistanceCondition->TargetActor = InTargetActor;
+
+	return SwitchOffDistanceCondition;
+}
+
+bool UMotionWarpingSwitchOffDistanceCondition::OnCheck() const
+{
+	switch (AxesType)
+	{
+		case ESwitchOffConditionDistanceAxesType::AllAxes:
+			return Operator == ESwitchOffConditionDistanceOp::LessThan ?
+				CalculateSqDistance() < FMath::Square(Distance) :
+				CalculateSqDistance() > FMath::Square(Distance);
+		
+		case ESwitchOffConditionDistanceAxesType::IgnoreZAxis:
+			return Operator == ESwitchOffConditionDistanceOp::LessThan ?
+				CalculateSqDistance2D() < FMath::Square(Distance) :
+				CalculateSqDistance2D() > FMath::Square(Distance);
+		
+		case ESwitchOffConditionDistanceAxesType::OnlyZAxis:
+			return Operator == ESwitchOffConditionDistanceOp::LessThan ?
+				CalculateZDistance() < Distance :
+				CalculateZDistance() > Distance;
+	}
+
+	return false;
+}
+
+float UMotionWarpingSwitchOffDistanceCondition::CalculateSqDistance() const
+{
+	return (OwnerActor->GetActorLocation() - GetTargetLocation()).SquaredLength();
+}
+
+float UMotionWarpingSwitchOffDistanceCondition::CalculateSqDistance2D() const
+{
+	return (OwnerActor->GetActorLocation() - GetTargetLocation()).SizeSquared2D();
+}
+
+float UMotionWarpingSwitchOffDistanceCondition::CalculateZDistance() const
+{
+	return FMath::Abs(OwnerActor->GetActorLocation().Z - GetTargetLocation().Z);
+}
+
+FString UMotionWarpingSwitchOffDistanceCondition::ExtraDebugInfo() const
+{
+	switch (AxesType)
+	{
+		case ESwitchOffConditionDistanceAxesType::AllAxes:
+			return FString::Printf(TEXT("Distance: %f %c %f"),
+				FMath::Sqrt(CalculateSqDistance()),
+				(Operator == ESwitchOffConditionDistanceOp::GreaterThan ? '>' : '<'),
+				Distance);
+		
+		case ESwitchOffConditionDistanceAxesType::IgnoreZAxis:
+			return FString::Printf(TEXT("Distance2D: %f %c %f"),
+				FMath::Sqrt(CalculateSqDistance2D()),
+				(Operator == ESwitchOffConditionDistanceOp::GreaterThan ? '>' : '<'),
+				Distance);
+		
+		case ESwitchOffConditionDistanceAxesType::OnlyZAxis:
+			return FString::Printf(TEXT("Distance Z: %f %c %f"),
+				CalculateZDistance(),
+				(Operator == ESwitchOffConditionDistanceOp::GreaterThan ? '>' : '<'),
+				Distance);
+	}
+
+	return FString();
+}
+
+UMotionWarpingSwitchOffAngleToTargetCondition* UMotionWarpingSwitchOffAngleToTargetCondition::CreateSwitchOffAngleToTargetCondition(AActor* InOwnerActor, ESwitchOffConditionEffect InEffect, ESwitchOffConditionAngleOp InOperator, float InAngle, bool bInIgnoreZAxis, bool bInUseWarpTargetAsTargetLocation, AActor* InTargetActor)
+{
+	UMotionWarpingSwitchOffAngleToTargetCondition* SwitchOffAngleToTargetCondition = NewObject<UMotionWarpingSwitchOffAngleToTargetCondition>();
+	SwitchOffAngleToTargetCondition->OwnerActor = InOwnerActor;
+	SwitchOffAngleToTargetCondition->Effect = InEffect;
+	SwitchOffAngleToTargetCondition->Operator = InOperator;
+	SwitchOffAngleToTargetCondition->Angle = InAngle;
+	SwitchOffAngleToTargetCondition->bUseWarpTargetAsTargetLocation = bInUseWarpTargetAsTargetLocation;
+	SwitchOffAngleToTargetCondition->TargetActor = InTargetActor;
+
+	return SwitchOffAngleToTargetCondition;
+}
+
+bool UMotionWarpingSwitchOffAngleToTargetCondition::OnCheck() const
+{
+	if (Operator == ESwitchOffConditionAngleOp::LessThan)
+	{
+		return CalculateAngleToTarget() < Angle;
+	}
+
+	return CalculateAngleToTarget() > Angle;
+}
+
+FString UMotionWarpingSwitchOffAngleToTargetCondition::ExtraDebugInfo() const
+{
+	return FString::Printf(TEXT("Angle: %f %c %f"),
+		CalculateAngleToTarget(),
+		(Operator == ESwitchOffConditionAngleOp::GreaterThan ? '>' : '<'),
+		Angle);
+}
+
+float UMotionWarpingSwitchOffAngleToTargetCondition::CalculateAngleToTarget() const
+{
+	FVector OwnerForward = OwnerActor->GetActorForwardVector();
+	FVector OwnerToTarget = GetTargetLocation() - OwnerActor->GetActorLocation();
+
+	if (bIgnoreZAxis)
+	{
+		OwnerForward = FVector(OwnerForward.X, OwnerForward.Y, 0);
+		OwnerToTarget = FVector(OwnerToTarget.X, OwnerToTarget.Y, 0);
+		OwnerForward.Normalize();
+	}
+
+	OwnerToTarget.Normalize();
+
+	return FMath::RadiansToDegrees(FMath::Acos(FVector::DotProduct(OwnerForward, OwnerToTarget)));
+}
+
+void UMotionWarpingSwitchOffCompositeCondition::SetOwnerActor(const AActor* InOwnerActor)
+{
+	Super::SetOwnerActor(InOwnerActor);
+	
+	SwitchOffConditionA->SetOwnerActor(InOwnerActor);
+	SwitchOffConditionB->SetOwnerActor(InOwnerActor);
+}
+
+void UMotionWarpingSwitchOffCompositeCondition::SetTargetActor(const AActor* InTargetActor)
+{
+	Super::SetTargetActor(InTargetActor);
+	
+	SwitchOffConditionA->SetTargetActor(InTargetActor);
+	SwitchOffConditionB->SetTargetActor(InTargetActor);
+}
+
+void UMotionWarpingSwitchOffCompositeCondition::SetMotionWarpingTarget(const FMotionWarpingTarget* InMotionWarpingTarget)
+{
+	Super::SetMotionWarpingTarget(InMotionWarpingTarget);
+	
+	SwitchOffConditionA->SetMotionWarpingTarget(InMotionWarpingTarget);
+	SwitchOffConditionB->SetMotionWarpingTarget(InMotionWarpingTarget);
+}
+
+void UMotionWarpingSwitchOffCompositeCondition::SetWarpTargetForDestination(const FMotionWarpingTarget* InMotionWarpingTarget)
+{
+	Super::SetWarpTargetForDestination(InMotionWarpingTarget);
+
+	if (ensureMsgf(IsValid(SwitchOffConditionA), TEXT("Switch off condition A not setup in composite switch off condition on actor %s"), *OwnerActor->GetName())
+		&& ensureMsgf(IsValid(SwitchOffConditionB), TEXT("Switch off condition B not setup in composite switch off condition on actor %s"), *OwnerActor->GetName()))
+	{
+		SwitchOffConditionA->SetWarpTargetForDestination(InMotionWarpingTarget);
+		SwitchOffConditionB->SetWarpTargetForDestination(InMotionWarpingTarget);
+	}
+}
+
+bool UMotionWarpingSwitchOffCompositeCondition::OnCheck() const
+{
+	if (ensureMsgf(IsValid(SwitchOffConditionA), TEXT("Switch off condition A not setup in composite switch off condition on actor %s"), *OwnerActor->GetName())
+		&& ensureMsgf(IsValid(SwitchOffConditionB), TEXT("Switch off condition B not setup in composite switch off condition on actor %s"), *OwnerActor->GetName()))
+	{
+		if (LogicOperator == ESwitchOffConditionCompositeOp::Or)
+		{
+			return SwitchOffConditionA->Check() || SwitchOffConditionB->Check();
+		}
+		return SwitchOffConditionA->Check() && SwitchOffConditionB->Check();
+	}
+
+	return false;
+}
+
+FString UMotionWarpingSwitchOffCompositeCondition::ExtraDebugInfo() const
+{
+	return FString::Printf(TEXT("%s %s %s"),
+		*SwitchOffConditionA->ExtraDebugInfo(),
+		*(LogicOperator == ESwitchOffConditionCompositeOp::Or ? FString("OR") : FString("AND")),
+		*SwitchOffConditionB->ExtraDebugInfo());
+}
+
+bool UMotionWarpingSwitchOffCompositeCondition::IsConditionValid() const
+{
+	return IsValid(SwitchOffConditionA) && IsValid(SwitchOffConditionB)
+		&& SwitchOffConditionA->IsConditionValid() && SwitchOffConditionB->IsConditionValid();
+}
+
+UMotionWarpingSwitchOffCompositeCondition* UMotionWarpingSwitchOffCompositeCondition::CreateSwitchOffCompositeCondition(AActor* InOwnerActor, ESwitchOffConditionEffect InEffect, UMotionWarpingSwitchOffCondition* InSwitchOffConditionA, ESwitchOffConditionCompositeOp InLogicOperator, UMotionWarpingSwitchOffCondition* InSwitchOffConditionB, bool bInUseWarpTargetAsTargetLocation, AActor* InTargetActor)
+{
+	UMotionWarpingSwitchOffCompositeCondition* SwitchOffCompositeCondition = NewObject<UMotionWarpingSwitchOffCompositeCondition>();
+	SwitchOffCompositeCondition->OwnerActor = InOwnerActor;
+	SwitchOffCompositeCondition->Effect = InEffect;
+	SwitchOffCompositeCondition->SwitchOffConditionA = InSwitchOffConditionA;
+	SwitchOffCompositeCondition->LogicOperator = InLogicOperator;
+	SwitchOffCompositeCondition->SwitchOffConditionB = InSwitchOffConditionB;
+	SwitchOffCompositeCondition->bUseWarpTargetAsTargetLocation = bInUseWarpTargetAsTargetLocation;
+	SwitchOffCompositeCondition->TargetActor = InTargetActor;
+
+	return SwitchOffCompositeCondition;
+}
+
+bool UMotionWarpingSwitchOffBlueprintableCondition::OnCheck() const
+{
+	return BP_Check(OwnerActor, TargetActor, GetTargetLocation(), bUseWarpTargetAsTargetLocation);
+}
+
+FString UMotionWarpingSwitchOffBlueprintableCondition::ExtraDebugInfo() const
+{
+	return BP_ExtraDebugInfo(OwnerActor, TargetActor, GetTargetLocation(), bUseWarpTargetAsTargetLocation);
+}
+
+UMotionWarpingSwitchOffBlueprintableCondition* UMotionWarpingSwitchOffBlueprintableCondition::CreateSwitchOffBlueprintableCondition(AActor* InOwnerActor, ESwitchOffConditionEffect InEffect, TSubclassOf<UMotionWarpingSwitchOffBlueprintableCondition> InBlueprintableCondition, bool bInUseWarpTargetAsTargetLocation, AActor* InTargetActor)
+{
+	UMotionWarpingSwitchOffBlueprintableCondition* SwitchOffBlueprintableCondition = NewObject<UMotionWarpingSwitchOffBlueprintableCondition>(GetTransientPackage(), InBlueprintableCondition);
+	SwitchOffBlueprintableCondition->OwnerActor = InOwnerActor;
+	SwitchOffBlueprintableCondition->Effect = InEffect;
+	SwitchOffBlueprintableCondition->bUseWarpTargetAsTargetLocation = bInUseWarpTargetAsTargetLocation;
+	SwitchOffBlueprintableCondition->TargetActor = InTargetActor;
+
+	return SwitchOffBlueprintableCondition;
+}
+
+UWorld* UMotionWarpingSwitchOffBlueprintableCondition::GetWorld() const
+{
+	if (IsValid(OwnerActor))
+	{
+		return OwnerActor->GetWorld();
+	}
+
+	return nullptr;
+}
+
+FString UMotionWarpingSwitchOffBlueprintableCondition::BP_ExtraDebugInfo_Implementation(const AActor* InOwnerActor, const AActor* InTargetActor, FVector InTargetLocation, bool bInUseWarpTargetAsTargetLocation) const
+{
+	return FString("No extra debug info. Override BP_ExtraDebugInfo to add it.");
+}
+
+bool UMotionWarpingSwitchOffBlueprintableCondition::BP_Check_Implementation(const AActor* InOwnerActor, const AActor* InTargetActor, FVector InTargetLocation, bool bInUseWarpTargetAsTargetLocation) const
+{
+	return false;
+}
diff --git a/Engine/Plugins/Animation/MotionWarping/Source/MotionWarping/Private/RootMotionModifier.cpp b/Engine/Plugins/Animation/MotionWarping/Source/MotionWarping/Private/RootMotionModifier.cpp
--- a/Engine/Plugins/Animation/MotionWarping/Source/MotionWarping/Private/RootMotionModifier.cpp
+++ b/Engine/Plugins/Animation/MotionWarping/Source/MotionWarping/Private/RootMotionModifier.cpp
@@ -324,6 +324,9 @@ void URootMotionModifier_Warp::Update(const FMotionWarpingUpdateContext& Context
 			return;
 		}
 
+		bRootMotionPaused = WarpTargetPtr->bRootMotionPaused;
+		bWarpingPaused = WarpTargetPtr->bWarpingPaused;
+		
 		// Get the warp point sent by the game
 		FTransform WarpPointTransformGame = WarpTargetPtr->GetTargetTrasform();
 
@@ -414,6 +417,16 @@ FQuat URootMotionModifier_Warp::GetTargetRotation() const
 
 FQuat URootMotionModifier_Warp::WarpRotation(const FTransform& RootMotionDelta, const FTransform& RootMotionTotal, float DeltaSeconds)
 {
+	if (bRootMotionPaused)
+	{
+		return FQuat::Identity;
+	}
+
+	if (bWarpingPaused)
+	{
+		return RootMotionDelta.GetRotation();
+	}
+	
 	FQuat CurrentRotation;
 	FQuat TargetRotation;
 
diff --git a/Engine/Plugins/Animation/MotionWarping/Source/MotionWarping/Private/RootMotionModifier_SkewWarp.cpp b/Engine/Plugins/Animation/MotionWarping/Source/MotionWarping/Private/RootMotionModifier_SkewWarp.cpp
--- a/Engine/Plugins/Animation/MotionWarping/Source/MotionWarping/Private/RootMotionModifier_SkewWarp.cpp
+++ b/Engine/Plugins/Animation/MotionWarping/Source/MotionWarping/Private/RootMotionModifier_SkewWarp.cpp
@@ -158,6 +158,12 @@ FTransform URootMotionModifier_SkewWarp::ProcessRootMotion(const FTransform& InR
 
 	if (bWarpTranslation)
 	{
+		if (bRootMotionPaused)
+		{
+			 FinalRootMotion.SetTranslation(FVector::ZeroVector);
+		}
+		else if (!bWarpingPaused)
+		{
 		const FVector CurrentLocation       = OwnerAdapter->GetVisualRootLocation();
 		const FQuat CurrentRotation         = OwnerAdapter->GetActor()->GetActorQuat();
 		FVector MeshTranslationOffset       = OwnerAdapter->GetBaseVisualTranslationOffset();
@@ -211,6 +217,7 @@ FTransform URootMotionModifier_SkewWarp::ProcessRootMotion(const FTransform& InR
 			}
 		}
 	}
+	}
 
 	if(bWarpRotation)
 	{
diff --git a/Engine/Plugins/Animation/MotionWarping/Source/MotionWarping/Public/MotionWarpingComponent.h b/Engine/Plugins/Animation/MotionWarping/Source/MotionWarping/Public/MotionWarpingComponent.h
--- a/Engine/Plugins/Animation/MotionWarping/Source/MotionWarping/Public/MotionWarpingComponent.h
+++ b/Engine/Plugins/Animation/MotionWarping/Source/MotionWarping/Public/MotionWarpingComponent.h
@@ -6,6 +6,8 @@
 #include "RootMotionModifier.h"
 #include "MotionWarpingComponent.generated.h"
 
+class UMotionWarpingSwitchOffCondition;
+struct FSwitchOffConditionData;
 template <class T> class TAutoConsoleVariable;
 
 class ACharacter;
@@ -26,6 +28,7 @@ struct MOTIONWARPING_API FMotionWarpingCVars
 	static TAutoConsoleVariable<int32> CVarMotionWarpingDebug;
 	static TAutoConsoleVariable<float> CVarMotionWarpingDrawDebugDuration;
 	static TAutoConsoleVariable<int32> CVarWarpedTargetDebug;
+	static TAutoConsoleVariable<int32> CVarWarpedSwitchOffConditionDebug;
 };
 #endif
 
@@ -206,6 +209,13 @@ public:
 	UFUNCTION(BlueprintCallable, Category = "Motion Warping")
 	int32 RemoveWarpTargets(const TArray<FName>& WarpTargetNames);
 
+	UFUNCTION(BlueprintCallable, Category="Motion Warping")
+	void AddSwitchOffCondition(FName WarpTargetName, UMotionWarpingSwitchOffCondition* Condition);
+	
+	void RemoveSwitchOffConditions(FName WarpTargetName);
+	
+	FSwitchOffConditionData* FindSwitchOffConditionData(FName WarpTargetName);
+
 protected:
 
 	/** Adapter that connects motion warping to an owner */
@@ -218,12 +228,17 @@ protected:
 
 	UPROPERTY(Transient, Replicated)
 	TArray<FMotionWarpingTarget> WarpTargets;
+		
+	UPROPERTY(Transient)
+    TArray<FSwitchOffConditionData> SwitchOffConditions;
 
 #if !(UE_BUILD_SHIPPING || UE_BUILD_TEST)
 	TOptional<FVector> OriginalRootMotionAccum;
 	TOptional<FVector> WarpedRootMotionAccum;
 #endif
 
+	void UpdateSwitchOffConditions();
+
 	void UpdateWithContext(const FMotionWarpingUpdateContext& Context, float DeltaSeconds);
 
 	bool FindAndUpdateWarpTarget(const FMotionWarpingTarget& WarpTarget);
diff --git a/Engine/Plugins/Animation/MotionWarping/Source/MotionWarping/Public/MotionWarpingSwitchOffCondition.h b/Engine/Plugins/Animation/MotionWarping/Source/MotionWarping/Public/MotionWarpingSwitchOffCondition.h
new file mode 100644
--- /dev/null
+++ b/Engine/Plugins/Animation/MotionWarping/Source/MotionWarping/Public/MotionWarpingSwitchOffCondition.h
@@ -0,0 +1,324 @@
+// Copyright Epic Games, Inc. All Rights Reserved.
+
+#pragma once
+
+#include "RootMotionModifier.h"
+#include "MotionWarpingSwitchOffCondition.generated.h"
+
+class UMotionWarpingSwitchOffCondition;
+struct FMotionWarpingTarget;
+
+USTRUCT(BlueprintType)
+struct MOTIONWARPING_API FSwitchOffConditionData
+{
+	GENERATED_BODY()
+
+	FSwitchOffConditionData()
+		: WarpTargetName(NAME_None), SwitchOffConditions(TArray<UMotionWarpingSwitchOffCondition*>())
+	{}
+	
+	FSwitchOffConditionData(FName InWarpTargetName, UMotionWarpingSwitchOffCondition* InSwitchOffCondition)
+		: WarpTargetName(InWarpTargetName)
+	{
+		SwitchOffConditions.Add(InSwitchOffCondition);
+	}
+
+	FSwitchOffConditionData(FName InWarpTargetName)
+		: WarpTargetName(InWarpTargetName), SwitchOffConditions(TArray<UMotionWarpingSwitchOffCondition*>())
+	{}
+	
+	UPROPERTY(BlueprintReadWrite, Category = "Defaults")
+	FName WarpTargetName;
+
+	UPROPERTY(BlueprintReadWrite, Category = "Defaults")
+	TArray<TObjectPtr<UMotionWarpingSwitchOffCondition>> SwitchOffConditions;
+
+	void SetMotionWarpingTarget(const FMotionWarpingTarget* MotionWarpingTarget);
+};
+
+/**Result of switch off condition.*/
+UENUM(BlueprintType, Category = "MotionWarping")
+enum class ESwitchOffConditionEffect : uint8
+{
+	/**Changes associated motion warping target from component to a location of this component
+	* in the frame in which this switch off condition appeared */
+	CancelFollow UMETA(DisplayName = "Cancel Follow Component"),
+
+	/**Removes associated motion warping target*/
+	CancelWarping UMETA(DisplayName = "Cancel Warping"),
+
+	/**During time slot in the animation, where switch off condition is true, only play root motion, without warping*/
+	PauseWarping UMETA(DisplayName = "Pause Warping"),
+
+	/**During time slot in the animation, where switch off condition is true, play the animation in place*/
+	PauseRootMotion UMETA(DisplayName = "Pause Root Motion")
+};
+
+UCLASS(Abstract, EditInlineNew, ClassGroup = MotionWarping)
+class MOTIONWARPING_API UMotionWarpingSwitchOffCondition : public UObject
+{
+	GENERATED_BODY()
+
+public:
+	/** If set to false, switch off condition will use target actor location */
+	UPROPERTY(EditAnywhere, Category = "MotionWarping")
+	bool bUseWarpTargetAsTargetLocation = true;
+
+	UPROPERTY(EditAnywhere, Category = "MotionWarping")
+	ESwitchOffConditionEffect Effect = ESwitchOffConditionEffect::CancelWarping;
+	
+	/**
+	 *	Initialize switch off condition
+	 *	@param InOwnerActor Actor with MotionWarpingComponent, owning corresponding WarpingTarget
+	 *	@param InTargetActor Target actor used for calculating switch off condition result
+	 */
+	ESwitchOffConditionEffect GetEffect() const { return Effect; }
+	bool Check() const;
+	
+	/**
+	 * If bUseWarpTargetAsTargetLocation is true, this will return target FMotionWarpingTarget location.
+	 * If bUseWarpTargetAsTargetLocation is false, this will return target actor location.
+	 */
+	FVector GetTargetLocation() const;
+
+	/**
+	 * Set warp target as context for calculating switch off condition result
+	 * if bUseWarpTargetAsTargetLocation is set to true
+	 */
+	virtual void SetWarpTargetForDestination(const FMotionWarpingTarget* InMotionWarpingTarget);
+	virtual bool OnCheck() const { return false; }
+
+	/**
+	 * Extra information used for debugging, if CVAR MotionWarping.Debug.SwitchOffCondition is true
+	 */
+	virtual FString ExtraDebugInfo() const { return FString(); }
+	virtual bool IsConditionValid() const;
+
+	virtual void SetOwnerActor(const AActor* InOwnerActor) { OwnerActor = InOwnerActor; }
+	virtual void SetTargetActor(const AActor* InTargetActor) { TargetActor = InTargetActor; }
+	virtual void SetMotionWarpingTarget(const FMotionWarpingTarget* InMotionWarpingTarget) { MotionWarpingTarget = InMotionWarpingTarget; }
+	
+protected:
+	bool bIsInitialized = false;
+
+	UPROPERTY()
+	TObjectPtr<const AActor> OwnerActor;
+
+	UPROPERTY()
+	TObjectPtr<const AActor> TargetActor;
+
+	const FMotionWarpingTarget* MotionWarpingTarget = nullptr;
+};
+
+UENUM(BlueprintType, Category = "MotionWarping")
+enum class ESwitchOffConditionDistanceOp : uint8
+{
+	LessThan,
+	GreaterThan
+};
+
+UENUM(BlueprintType, Category = "MotionWarping")
+enum class ESwitchOffConditionDistanceAxesType : uint8
+{
+	AllAxes,
+	IgnoreZAxis,
+	OnlyZAxis,
+};
+
+UCLASS()
+class UMotionWarpingSwitchOffDistanceCondition : public UMotionWarpingSwitchOffCondition
+{
+	GENERATED_BODY()
+
+public:
+	UPROPERTY(EditAnywhere, Category = "MotionWarping")
+	ESwitchOffConditionDistanceOp Operator;
+
+	UPROPERTY(EditAnywhere, Category = "MotionWarping")
+	ESwitchOffConditionDistanceAxesType AxesType;
+	
+	UPROPERTY(EditAnywhere, Category = "MotionWarping")
+	float Distance;
+
+	/**
+	 * Creates Switch Off Distance Condition that checks distance between Owner Actor and Target Location.
+	 * If Use Warp Target Location is true, Target Location is corresponding Warp Target's location.
+	 * If Use Warp Target Location is false, Target Location is Target Actor's parameter location.
+	 * @param InOwnerActor Actor owning Motion Warping Component to which this switch off condition will be added.
+	 * @param InEffect What should happen if switch off condition is true
+	 * @param InOperator Distance comparison operator
+	 * @param InDistance Distance to compare to
+	 * @param InbUseWarpTargetAsTargetLocation Should use corresponding warp target as context for this Switch Off Condition.
+	 * If set to false, will use TargetActor parameter
+	 * @param InTargetActor Target Actor to use as a context for this Switch Off Condition if Use Warp Target As Location parameter is false.
+	 * @return Created Switch Off Condition. This Switch Off Condition can be added to Motion Warping Component with Add Switch Off Condition node.
+	 */
+	UFUNCTION(BlueprintCallable, Category="Motion Warping", meta=(AdvancedDisplay="bInUseWarpTargetAsTargetLocation, InTargetActor"))
+	static UMotionWarpingSwitchOffDistanceCondition* CreateSwitchOffDistanceCondition(
+		UPARAM(DisplayName = "Owner Actor", ref)		AActor* InOwnerActor,
+		UPARAM(DisplayName = "Effect")					ESwitchOffConditionEffect InEffect,
+		UPARAM(DisplayName = "Operator")				ESwitchOffConditionDistanceOp InOperator,
+		UPARAM(DisplayName = "Distance")				float InDistance,
+		UPARAM(DisplayName = "UseWarpTargetAsLocation") bool InbUseWarpTargetAsTargetLocation = true,
+		UPARAM(DisplayName = "TargetActor")				AActor* InTargetActor = nullptr);
+	
+protected:
+	virtual bool OnCheck() const override;
+	virtual FString ExtraDebugInfo() const override;
+	
+	float CalculateSqDistance() const;
+	float CalculateSqDistance2D() const;
+	float CalculateZDistance() const;
+};
+
+
+UENUM(BlueprintType, Category = "MotionWarping")
+enum class ESwitchOffConditionAngleOp : uint8
+{
+	LessThan,
+	GreaterThan
+};
+
+UCLASS()
+class UMotionWarpingSwitchOffAngleToTargetCondition : public UMotionWarpingSwitchOffCondition
+{
+	GENERATED_BODY()
+
+public:
+	/**
+	 * Creates Switch Off Angle To Target Condition that checks angle between Owner Actor and Target Location.
+	 * If Use Warp Target Location is true, Target Location is corresponding Warp Target's location.
+	 * If Use Warp Target Location is false, Target Location is Target Actor's parameter location.
+	 * @param InOwnerActor Actor owning Motion Warping Component to which this switch off condition will be added.
+	 * @param InEffect What should happen if switch off condition is true
+	 * @param InOperator Angle comparison operator
+	 * @param InAngle Angle to compare to
+	 * @param bInIgnoreZAxis Should ignore Z axis in Angle comparison
+	 * @param bInUseWarpTargetAsTargetLocation Should use corresponding warp target as context for this Switch Off Condition.
+	 * If set to false, will use TargetActor parameter
+	 * @param InTargetActor Target Actor to use as a context for this Switch Off Condition if Use Warp Target As Location parameter is false.
+	 * @return Created Switch Off Condition. This Switch Off Condition can be added to Motion Warping Component with Add Switch Off Condition node.
+	 */
+	UFUNCTION(BlueprintCallable, Category="Motion Warping", meta=(AdvancedDisplay="bInUseWarpTargetAsTargetLocation, InTargetActor"))
+	static UMotionWarpingSwitchOffAngleToTargetCondition* CreateSwitchOffAngleToTargetCondition(
+		UPARAM(DisplayName = "Owner Actor", ref)		AActor* InOwnerActor,
+		UPARAM(DisplayName = "Effect")					ESwitchOffConditionEffect InEffect,
+		UPARAM(DisplayName = "Operator")				ESwitchOffConditionAngleOp InOperator,
+		UPARAM(DisplayName = "Angle")					float InAngle,
+		UPARAM(DisplayName = "Ignore Z Axis")			bool bInIgnoreZAxis,
+		UPARAM(DisplayName = "TargetActor")				bool bInUseWarpTargetAsTargetLocation = true,
+		UPARAM(DisplayName = "TargetActor")				AActor* InTargetActor = nullptr);
+	
+protected:
+	virtual bool OnCheck() const override;
+	virtual FString ExtraDebugInfo() const override;
+
+	float CalculateAngleToTarget() const;
+
+	UPROPERTY(EditAnywhere, Category = "MotionWarping")
+	ESwitchOffConditionAngleOp Operator;
+
+	UPROPERTY(EditAnywhere, Category = "MotionWarping", meta = (ForceUnits="degrees"))
+	float Angle;
+
+	UPROPERTY(EditAnywhere, Category = "MotionWarping")
+	bool bIgnoreZAxis;
+};
+
+UENUM(BlueprintType, Category = "MotionWarping")
+enum class ESwitchOffConditionCompositeOp : uint8
+{
+	Or UMETA(DisplayName = "OR"),
+	And UMETA(DisplayName = "AND")
+};
+
+UCLASS()
+class UMotionWarpingSwitchOffCompositeCondition : public UMotionWarpingSwitchOffCondition
+{
+	GENERATED_BODY()
+public:
+	virtual void SetOwnerActor(const AActor* InOwnerActor) override;
+	virtual void SetTargetActor(const AActor* InTargetActor) override;
+	virtual void SetMotionWarpingTarget(const FMotionWarpingTarget* InMotionWarpingTarget) override;
+
+	/**
+	 * Creates Switch Off Composite Condition that lets you combine different switch off conditions with logic AND/OR operators.
+	 * @param InOwnerActor Actor owning Motion Warping Component to which this switch off condition will be added.
+	 * @param InEffect What should happen if switch off condition is true
+	 * @param InSwitchOffConditionA First Switch Off Condition to combine
+	 * @param InLogicOperator Logic operator to use for Switch Off Condition combination
+	 * @param InSwitchOffConditionB Second Switch Off Condition to combine
+	 * @param bInUseWarpTargetAsTargetLocation Should use corresponding warp target as context for this Switch Off Condition.
+	 * If set to false, will use TargetActor parameter
+	 * @param InTargetActor Target Actor to use as a context for this Switch Off Condition if Use Warp Target As Location parameter is false.
+	 * @return Created Switch Off Condition. This Switch Off Condition can be added to Motion Warping Component with Add Switch Off Condition node.
+	 */
+	UFUNCTION(BlueprintCallable, Category="Motion Warping", meta=(AdvancedDisplay="bInUseWarpTargetAsTargetLocation, InTargetActor"))
+	static UMotionWarpingSwitchOffCompositeCondition* CreateSwitchOffCompositeCondition(
+		UPARAM(DisplayName = "Owner Actor", ref)			AActor* InOwnerActor,
+		UPARAM(DisplayName = "Effect")						ESwitchOffConditionEffect InEffect,
+		UPARAM(DisplayName = "Switch Off Condition A", ref)	UMotionWarpingSwitchOffCondition* InSwitchOffConditionA,
+		UPARAM(DisplayName = "Logic Operator")				ESwitchOffConditionCompositeOp InLogicOperator,
+		UPARAM(DisplayName = "Switch Off Condition B", ref)	UMotionWarpingSwitchOffCondition* InSwitchOffConditionB,
+		UPARAM(DisplayName = "Use Warp Target As Location")	bool bInUseWarpTargetAsTargetLocation = true,
+		UPARAM(DisplayName = "TargetActor")					AActor* InTargetActor = nullptr);
+	
+protected:
+	virtual void SetWarpTargetForDestination(const FMotionWarpingTarget* InMotionWarpingTarget) override;
+	virtual bool OnCheck() const override;
+	virtual FString ExtraDebugInfo() const override;
+	virtual bool IsConditionValid() const override;
+	
+	UPROPERTY(EditAnywhere, Instanced, Category="MotionWarping")
+	TObjectPtr<UMotionWarpingSwitchOffCondition> SwitchOffConditionA;
+
+	UPROPERTY(EditAnywhere, Category="MotionWarping")
+	ESwitchOffConditionCompositeOp LogicOperator;
+
+	UPROPERTY(EditAnywhere, Instanced, Category="MotionWarping")
+	TObjectPtr<UMotionWarpingSwitchOffCondition> SwitchOffConditionB;
+};
+
+UCLASS(Blueprintable)
+class UMotionWarpingSwitchOffBlueprintableCondition : public UMotionWarpingSwitchOffCondition
+{
+	GENERATED_BODY()
+	
+	virtual bool OnCheck() const override;
+	virtual FString ExtraDebugInfo() const override;
+
+	/**
+	 * Creates Switch Off Blueprintable Condition from WarpingSwitchOffBlueprintableCondition subclass.
+	 * @param InOwnerActor Actor owning Motion Warping Component to which this switch off condition will be added.
+	 * @param InEffect What should happen if switch off condition is true
+	 * @param InBlueprintableCondition Subclass of WarpingSwitchOffBlueprintableCondition
+	* @param bInUseWarpTargetAsTargetLocation Should use corresponding warp target as context for this Switch Off Condition.
+	 * If set to false, will use TargetActor parameter 
+	 * @param InTargetActor Target Actor to use as a context for this Switch Off Condition if Use Warp Target As Location parameter is false.
+	 * @return Created Switch Off Condition. This Switch Off Condition can be added to Motion Warping Component with Add Switch Off Condition node.
+	 */
+	UFUNCTION(BlueprintCallable, Category="Motion Warping", meta=(AdvancedDisplay="bInUseWarpTargetAsTargetLocation, InTargetActor"))
+	static UMotionWarpingSwitchOffBlueprintableCondition* CreateSwitchOffBlueprintableCondition(
+		UPARAM(DisplayName = "Owner Actor", ref)					AActor* InOwnerActor,
+		UPARAM(DisplayName = "Effect")								ESwitchOffConditionEffect InEffect,
+		UPARAM(DisplayName = "Blueprintable Switch Off Condition")	TSubclassOf<UMotionWarpingSwitchOffBlueprintableCondition> InBlueprintableCondition,
+		UPARAM(DisplayName = "Use Warp Target As Location")			bool bInUseWarpTargetAsTargetLocation = true,
+		UPARAM(DisplayName = "TargetActor")							AActor* InTargetActor = nullptr);
+	
+public:
+	virtual UWorld* GetWorld() const override;
+
+	UFUNCTION(BlueprintNativeEvent)
+	bool BP_Check(
+		UPARAM(DisplayName = "Owner Actor")							const AActor* InOwnerActor,
+		UPARAM(DisplayName = "Target Actor")						const AActor* InTargetActor,
+		UPARAM(DisplayName = "Target Location")						FVector InTargetLocation,
+		UPARAM(DisplayName = "Use warp target as target location")	bool bInUseWarpTargetAsTargetLocation) const;
+
+	UFUNCTION(BlueprintNativeEvent)
+	FString BP_ExtraDebugInfo(
+		UPARAM(DisplayName = "Owner Actor")							const AActor* InOwnerActor,
+		UPARAM(DisplayName = "Target Actor")						const AActor* InTargetActor,
+		UPARAM(DisplayName = "Target Location")						FVector InTargetLocation,
+		UPARAM(DisplayName = "Use warp target as target location")	bool bInUseWarpTargetAsTargetLocation) const;
+};
\ No newline at end of file
diff --git a/Engine/Plugins/Animation/MotionWarping/Source/MotionWarping/Public/RootMotionModifier.h b/Engine/Plugins/Animation/MotionWarping/Source/MotionWarping/Public/RootMotionModifier.h
--- a/Engine/Plugins/Animation/MotionWarping/Source/MotionWarping/Public/RootMotionModifier.h
+++ b/Engine/Plugins/Animation/MotionWarping/Source/MotionWarping/Public/RootMotionModifier.h
@@ -226,6 +226,9 @@ struct MOTIONWARPING_API FMotionWarpingTarget
 	FVector CachedRightOffset = FVector::ZeroVector;
 	FVector CachedUpOffset = FVector::ZeroVector;
 
+	bool bWarpingPaused = false;
+	bool bRootMotionPaused = false;
+
 	FMotionWarpingTarget()
 		: Name(NAME_None), Location(FVector::ZeroVector), Rotation(FRotator::ZeroRotator), Component(nullptr), BoneName(NAME_None), bFollowComponent(false) {}
 
@@ -400,6 +403,9 @@ protected:
 
 	/** Cached of the offset from the warp target. Used to calculate the final target transform when a warp target is defined in the animation */
 	TOptional<FTransform> CachedOffsetFromWarpPoint;
+
+	bool bWarpingPaused = false;
+	bool bRootMotionPaused = false;
 };
 
 // URootMotionModifier_SimpleWarp. 
