[Engine] [MotionWarping] Added max speed clamp ratio for translation warping.

Credit: oskar.kogut

#rnx

diff --git a/Engine/Plugins/Animation/MotionWarping/Source/MotionWarping/Private/RootMotionModifier_SkewWarp.cpp b/Engine/Plugins/Animation/MotionWarping/Source/MotionWarping/Private/RootMotionModifier_SkewWarp.cpp
--- a/Engine/Plugins/Animation/MotionWarping/Source/MotionWarping/Private/RootMotionModifier_SkewWarp.cpp
+++ b/Engine/Plugins/Animation/MotionWarping/Source/MotionWarping/Private/RootMotionModifier_SkewWarp.cpp
@@ -191,7 +191,20 @@ FTransform URootMotionModifier_SkewWarp::ProcessRootMotion(const FTransform& InR
 				const FTransform MeshTransform = FTransform(MeshRotationOffset, MeshTranslationOffset) * OwnerAsActor->GetActorTransform();
 				TargetLocation = MeshTransform.InverseTransformPositionNoScale(TargetLocation) - RootMotionRemainingAfterNotify.GetTranslation();
 
-				const FVector WarpedTranslation = WarpTranslation(FTransform::Identity, DeltaTranslation, TotalTranslation, TargetLocation) + ExtraRootMotion.GetLocation();
+				FVector WarpedTranslation = WarpTranslation(FTransform::Identity, DeltaTranslation, TotalTranslation, TargetLocation) + ExtraRootMotion.GetLocation();
+				if (MaxSpeedClampRatio > 0.0f)
+				{
+					if (bIgnoreZAxis)
+					{
+						const float AnimationSpeed = DeltaTranslation.Size2D();
+						WarpedTranslation = WarpedTranslation.GetClampedToMaxSize2D(AnimationSpeed * MaxSpeedClampRatio);
+					}
+					else
+					{
+						const float AnimationSpeed = DeltaTranslation.Size();
+						WarpedTranslation = WarpedTranslation.GetClampedToMaxSize(AnimationSpeed * MaxSpeedClampRatio);
+					}
+				}
 				FinalRootMotion.SetTranslation(WarpedTranslation);
 			}
 		}
diff --git a/Engine/Plugins/Animation/MotionWarping/Source/MotionWarping/Public/RootMotionModifier_SkewWarp.h b/Engine/Plugins/Animation/MotionWarping/Source/MotionWarping/Public/RootMotionModifier_SkewWarp.h
--- a/Engine/Plugins/Animation/MotionWarping/Source/MotionWarping/Public/RootMotionModifier_SkewWarp.h
+++ b/Engine/Plugins/Animation/MotionWarping/Source/MotionWarping/Public/RootMotionModifier_SkewWarp.h
@@ -41,4 +41,14 @@ public:
 		UPARAM(DisplayName = "Rotation Method") EMotionWarpRotationMethod InRotationMethod = EMotionWarpRotationMethod::Slerp,
 		UPARAM(DisplayName = "Warp Rotation Time Multiplier") float InWarpRotationTimeMultiplier = 1.f,
 		UPARAM(DisplayName = "Warp Max Rotation Rate") float InWarpMaxRotationRate = 0.f);
+
+protected:
+	/**
+	* Allows to set maximum warp translation speed clamp ratio.
+	* Ratio is relative to the original animation translation speed. E.g. if the MaxSpeedClampRatio == 2.0f actor well be moving with maximum 2x speed of the animation.
+	* Applied only in cases when animation has root motion translation.
+	* Zero treated as no clamping.
+	*/
+	UPROPERTY(EditAnywhere, BlueprintReadWrite, Category = "Config")
+	float MaxSpeedClampRatio = 0.0f;
 };
