Fix garbage reference - make UMotionWarpingCharacterAdapter take a weak reference to its target actor.

diff --git a/Engine/Plugins/Animation/MotionWarping/Source/MotionWarping/Private/MotionWarpingCharacterAdapter.cpp b/Engine/Plugins/Animation/MotionWarping/Source/MotionWarping/Private/MotionWarpingCharacterAdapter.cpp
--- a/Engine/Plugins/Animation/MotionWarping/Source/MotionWarping/Private/MotionWarpingCharacterAdapter.cpp
+++ b/Engine/Plugins/Animation/MotionWarping/Source/MotionWarping/Private/MotionWarpingCharacterAdapter.cpp
@@ -9,9 +9,11 @@
 
 void UMotionWarpingCharacterAdapter::BeginDestroy()
 {
-	if (TargetCharacter && TargetCharacter->GetCharacterMovement())
+	const ACharacter* RawTargetCharacter = TargetCharacter.Get();
+	UCharacterMovementComponent* CharacterMovement = RawTargetCharacter ? RawTargetCharacter->GetCharacterMovement() : nullptr;
+	if (CharacterMovement)
 	{
-		TargetCharacter->GetCharacterMovement()->ProcessRootMotionPreConvertToWorld.Unbind();
+		CharacterMovement->ProcessRootMotionPreConvertToWorld.Unbind();
 	}
 
 	Super::BeginDestroy();
@@ -22,49 +24,66 @@ void UMotionWarpingCharacterAdapter::SetCharacter(ACharacter* InCharacter)
 	if (ensureMsgf(InCharacter && InCharacter->GetCharacterMovement(), TEXT("Invalid Character or missing CharacterMovementComponent. Motion warping will not function.")))
 	{
 		TargetCharacter = InCharacter;
-		TargetCharacter->GetCharacterMovement()->ProcessRootMotionPreConvertToWorld.BindUObject(this, &UMotionWarpingCharacterAdapter::WarpLocalRootMotionOnCharacter);
+		InCharacter->GetCharacterMovement()->ProcessRootMotionPreConvertToWorld.BindUObject(this, &UMotionWarpingCharacterAdapter::WarpLocalRootMotionOnCharacter);
 	}
 }
 
 AActor* UMotionWarpingCharacterAdapter::GetActor() const
 { 
-	return Cast<AActor>(TargetCharacter);
+	return Cast<AActor>(TargetCharacter.Get());
 }
 
 USkeletalMeshComponent* UMotionWarpingCharacterAdapter::GetMesh() const
 { 
-	return TargetCharacter->GetMesh();
+	if (const ACharacter* RawTargetCharacter = TargetCharacter.Get())
+	{
+		return RawTargetCharacter->GetMesh();
+	}
+	return nullptr;
 }
 
 FVector UMotionWarpingCharacterAdapter::GetVisualRootLocation() const
 {
-	const float CapsuleHalfHeight = TargetCharacter->GetCapsuleComponent()->GetScaledCapsuleHalfHeight();
-	const FQuat CurrentRotation = TargetCharacter->GetActorQuat();
-	return  (TargetCharacter->GetActorLocation() - CurrentRotation.GetUpVector() * CapsuleHalfHeight);
+	if (const ACharacter* RawTargetCharacter = TargetCharacter.Get())
+	{
+		const float CapsuleHalfHeight = RawTargetCharacter->GetCapsuleComponent()->GetScaledCapsuleHalfHeight();
+		const FQuat CurrentRotation = RawTargetCharacter->GetActorQuat();
+		return  (RawTargetCharacter->GetActorLocation() - CurrentRotation.GetUpVector() * CapsuleHalfHeight);
+	}
+	return FVector::ZeroVector;
 }
 
 FVector UMotionWarpingCharacterAdapter::GetBaseVisualTranslationOffset() const
 {
-	return TargetCharacter->GetBaseTranslationOffset();
+	if (const ACharacter* RawTargetCharacter = TargetCharacter.Get())
+	{
+		return RawTargetCharacter->GetBaseTranslationOffset();
+	}
+	return FVector::ZeroVector;
 }
 
 FQuat UMotionWarpingCharacterAdapter::GetBaseVisualRotationOffset() const
 { 
-	return TargetCharacter->GetBaseRotationOffset();
+	if (const ACharacter* RawTargetCharacter = TargetCharacter.Get())
+	{
+		return RawTargetCharacter->GetBaseRotationOffset();
+	}
+	return FQuat::Identity;
 }
 
 FTransform UMotionWarpingCharacterAdapter::WarpLocalRootMotionOnCharacter(const FTransform& LocalRootMotionTransform, UCharacterMovementComponent* TargetMoveComp, float DeltaSeconds)
 {
-	if (WarpLocalRootMotionDelegate.IsBound() && TargetCharacter)
+	const ACharacter* RawTargetCharacter = TargetCharacter.Get();
+	if (WarpLocalRootMotionDelegate.IsBound() && RawTargetCharacter)
 	{
 		FMotionWarpingUpdateContext WarpingContext;
 		
 		WarpingContext.DeltaSeconds = DeltaSeconds;
 
 		// When replaying saved moves we need to look at the contributor to root motion back then.
-		if (TargetCharacter->bClientUpdating)
+		if (RawTargetCharacter->bClientUpdating)
 		{
-			const UCharacterMovementComponent* MoveComp = TargetCharacter->GetCharacterMovement();
+			const UCharacterMovementComponent* MoveComp = RawTargetCharacter->GetCharacterMovement();
 			check(MoveComp);
 
 			const FSavedMove_Character* SavedMove = MoveComp->GetCurrentReplayedSavedMove();
@@ -80,7 +99,7 @@ FTransform UMotionWarpingCharacterAdapter::WarpLocalRootMotionOnCharacter(const
 		}
 		else // If we are not replaying a move, just use the current root motion montage
 		{
-			if (const FAnimMontageInstance* RootMotionMontageInstance = TargetCharacter->GetRootMotionAnimMontageInstance())
+			if (const FAnimMontageInstance* RootMotionMontageInstance = RawTargetCharacter->GetRootMotionAnimMontageInstance())
 			{
 				const UAnimMontage* Montage = RootMotionMontageInstance->Montage;
 				check(Montage);
diff --git a/Engine/Plugins/Animation/MotionWarping/Source/MotionWarping/Public/MotionWarpingCharacterAdapter.h b/Engine/Plugins/Animation/MotionWarping/Source/MotionWarping/Public/MotionWarpingCharacterAdapter.h
--- a/Engine/Plugins/Animation/MotionWarping/Source/MotionWarping/Public/MotionWarpingCharacterAdapter.h
+++ b/Engine/Plugins/Animation/MotionWarping/Source/MotionWarping/Public/MotionWarpingCharacterAdapter.h
@@ -28,6 +28,5 @@ private:
 	FTransform WarpLocalRootMotionOnCharacter(const FTransform& LocalRootMotionTransform, UCharacterMovementComponent* TargetMoveComp, float DeltaSeconds);
 
 	/** The associated character */
-	UPROPERTY(Transient, DuplicateTransient)
-	TObjectPtr<ACharacter> TargetCharacter;
+	TWeakObjectPtr<ACharacter> TargetCharacter;
 };
\ No newline at end of file
