Motion Warping: Added missing logic for IgnoreZAxis when injecting root motion.

diff --git a/Engine/Plugins/Animation/MotionWarping/Source/MotionWarping/Private/RootMotionModifier_SkewWarp.cpp b/Engine/Plugins/Animation/MotionWarping/Source/MotionWarping/Private/RootMotionModifier_SkewWarp.cpp
--- a/Engine/Plugins/Animation/MotionWarping/Source/MotionWarping/Private/RootMotionModifier_SkewWarp.cpp
+++ b/Engine/Plugins/Animation/MotionWarping/Source/MotionWarping/Private/RootMotionModifier_SkewWarp.cpp
@@ -223,7 +223,12 @@ FTransform URootMotionModifier_SkewWarp::ProcessRootMotion(const FTransform& InR
 				float Alpha = FMath::Clamp((CurrentPosition - ActualStartTime) / (EndTime - ActualStartTime), 0.f, 1.f);
 				Alpha = FAlphaBlend::AlphaToBlendOption(Alpha, AddTranslationEasingFunc, AddTranslationEasingCurve);
 
-				const FVector NextLocation = FMath::Lerp<FVector, float>(StartTransform.GetLocation(), TargetLocation, Alpha);
+				FVector NextLocation = FMath::Lerp<FVector, float>(StartTransform.GetLocation(), TargetLocation, Alpha);
+				if (bIgnoreZAxis)
+				{
+					NextLocation.Z = CurrentLocation.Z;
+				}
+
 				FVector FinalDeltaTranslation = (NextLocation - CurrentLocation);
 				FinalDeltaTranslation = (CurrentRotation.Inverse() * DeltaToTarget.ToOrientationQuat()).GetForwardVector() * FinalDeltaTranslation.Size();
 				FinalDeltaTranslation = MeshRotationOffset.UnrotateVector(FinalDeltaTranslation);
