[Animation] Motion warping "scale" rotation method.
Scales the root motion rotation in a way that by the end of the corresponding anim notify state the owner's rotation matches the warp target's rotation.

diff --git a/Engine/Plugins/Animation/MotionWarping/Source/MotionWarping/Private/RootMotionModifier.cpp b/Engine/Plugins/Animation/MotionWarping/Source/MotionWarping/Private/RootMotionModifier.cpp
--- a/Engine/Plugins/Animation/MotionWarping/Source/MotionWarping/Private/RootMotionModifier.cpp
+++ b/Engine/Plugins/Animation/MotionWarping/Source/MotionWarping/Private/RootMotionModifier.cpp
@@ -442,6 +442,24 @@ FQuat URootMotionModifier_Warp::WarpRotation(const FTransform& RootMotionDelta,
 	}
 
 	const FQuat TotalRootMotionRotation = RootMotionTotal.GetRotation();
+
+	if (RotationMethod == EMotionWarpRotationMethod::Scale)
+	{
+		FRotator TotalRotator(TotalRootMotionRotation);
+		FRotator TargetRotator(TargetRotation);
+		const double YawDiff = FMath::FindDeltaAngleDegrees(TotalRotator.Yaw, TargetRotator.Yaw);
+		const double PitchDiff = FMath::FindDeltaAngleDegrees(TotalRotator.Pitch, TargetRotator.Pitch);
+		// To properly compute scale factor target rotation needs to be relative to total rotation.
+		// To avoid cases like 170 & -170 resulting in -1 scale factor rather than 1.11.
+		const double YawScale = FMath::IsNearlyZero(TotalRotator.Yaw) ? 0.0 : (TotalRotator.Yaw + YawDiff) / TotalRotator.Yaw;
+		const double PitchScale = FMath::IsNearlyZero(TotalRotator.Pitch) ? 0.0 : (TotalRotator.Pitch + PitchDiff) / TotalRotator.Pitch;
+		const float MaxRotation = WarpMaxRotationRate * DeltaSeconds;
+		FRotator ScaledDeltaRotation(RootMotionDelta.GetRotation());
+		ScaledDeltaRotation.Yaw = FMath::Clamp(ScaledDeltaRotation.Yaw * YawScale, -MaxRotation, MaxRotation);
+		ScaledDeltaRotation.Pitch = FMath::Clamp(ScaledDeltaRotation.Pitch * PitchScale, -MaxRotation, MaxRotation);
+		return ScaledDeltaRotation.Quaternion();
+	}
+
 	const float TimeRemaining = (EndTime - PreviousPosition) * WarpRotationTimeMultiplier;
 	const float Alpha = FMath::Clamp(DeltaSeconds / TimeRemaining, 0.f, 1.f);
 	FQuat TargetRotThisFrame = FQuat::Slerp(TotalRootMotionRotation, TargetRotation, Alpha);
diff --git a/Engine/Plugins/Animation/MotionWarping/Source/MotionWarping/Public/RootMotionModifier.h b/Engine/Plugins/Animation/MotionWarping/Source/MotionWarping/Public/RootMotionModifier.h
--- a/Engine/Plugins/Animation/MotionWarping/Source/MotionWarping/Public/RootMotionModifier.h
+++ b/Engine/Plugins/Animation/MotionWarping/Source/MotionWarping/Public/RootMotionModifier.h
@@ -286,6 +286,8 @@ enum class EMotionWarpRotationMethod : uint8
 	SlerpWithClampedRate,
 	/** Rotate with a constant rotation rate*/
 	ConstantRate,
+	/** Scale rotation in a way that by the end of the anim notify state the owner's rotation matches the warp target's rotation */
+	Scale,
 };
 
 /** Method used to extract the warp point from the animation */
@@ -365,7 +367,7 @@ public:
 	 * Allow to modify how fast the rotation is warped.
 	 * e.g if the window duration is 2sec and this is 0.5, the target rotation will be reached in 1sec instead of 2sec
 	 */
-	UPROPERTY(EditAnywhere, BlueprintReadWrite, Category = "Config", meta = (EditCondition = "RotationMethod!=EMotionWarpRotationMethod::ConstantRate && bWarpRotation"))
+	UPROPERTY(EditAnywhere, BlueprintReadWrite, Category = "Config", meta = (EditCondition = "RotationMethod!=EMotionWarpRotationMethod::ConstantRate && RotationMethod!=EMotionWarpRotationMethod::Scale && bWarpRotation"))
 	float WarpRotationTimeMultiplier = 1.f;
 
 	/** Maximum rotation rate in degrees/sec. Will be the value used in constant rotation rate*/
