Motion Warping - Added blueprint callable function to extract the bone transform from an animation at any given time

diff --git a/Engine/Plugins/Animation/MotionWarping/Source/MotionWarping/Private/MotionWarpingComponent.cpp b/Engine/Plugins/Animation/MotionWarping/Source/MotionWarping/Private/MotionWarpingComponent.cpp
--- a/Engine/Plugins/Animation/MotionWarping/Source/MotionWarping/Private/MotionWarpingComponent.cpp
+++ b/Engine/Plugins/Animation/MotionWarping/Source/MotionWarping/Private/MotionWarpingComponent.cpp
@@ -249,6 +249,32 @@ FTransform UMotionWarpingUtilities::CalculateRootTransformRelativeToWarpPointAtT
 	return RootTransform.GetRelativeTransform((MeshCompRelativeRotInverse * WarpPointTransform));
 }
 
+void UMotionWarpingUtilities::ExtractBoneTransformFromAnimationAtTime(const UAnimInstance* AnimInstance, const UAnimSequenceBase* Animation, float Time, bool bExtractRootMotion, FName BoneName, bool bLocalSpace, FTransform& OutTransform)
+{
+	OutTransform = FTransform::Identity;
+
+	if (AnimInstance && Animation)
+	{
+		FMemMark Mark(FMemStack::Get());
+
+		const int32 BoneIndex = AnimInstance->GetRequiredBones().GetPoseBoneIndexForBoneName(BoneName);
+		if (BoneIndex != INDEX_NONE)
+		{
+			if (bLocalSpace)
+			{
+				FCompactPose Pose;
+				ExtractLocalSpacePose(Animation, AnimInstance->GetRequiredBones(), Time, bExtractRootMotion, Pose);
+				OutTransform = Pose[FCompactPoseBoneIndex(BoneIndex)];
+			}
+			else
+			{
+				FCSPose<FCompactPose> Pose;
+				UMotionWarpingUtilities::ExtractComponentSpacePose(Animation, AnimInstance->GetRequiredBones(), Time, bExtractRootMotion, Pose);
+				OutTransform = Pose.GetComponentSpaceTransform(FCompactPoseBoneIndex(BoneIndex));
+			}
+		}
+	}
+}
 
 // UMotionWarpingComponent
 ///////////////////////////////////////////////////////////////////////
diff --git a/Engine/Plugins/Animation/MotionWarping/Source/MotionWarping/Public/MotionWarpingComponent.h b/Engine/Plugins/Animation/MotionWarping/Source/MotionWarping/Public/MotionWarpingComponent.h
--- a/Engine/Plugins/Animation/MotionWarping/Source/MotionWarping/Public/MotionWarpingComponent.h
+++ b/Engine/Plugins/Animation/MotionWarping/Source/MotionWarping/Public/MotionWarpingComponent.h
@@ -11,6 +11,7 @@ struct FSwitchOffConditionData;
 template <class T> class TAutoConsoleVariable;
 
 class ACharacter;
+class UAnimInstance;
 class UAnimSequenceBase;
 class UCharacterMovementComponent;
 class UMotionWarpingBaseAdapter;
@@ -83,6 +84,10 @@ public:
 	/** @return root transform relative to the warp point transform at the supplied time */
 	static FTransform CalculateRootTransformRelativeToWarpPointAtTime(const ACharacter& Character, const UAnimSequenceBase* Animation, float Time, const FTransform& WarpPointTransform);
 	static FTransform CalculateRootTransformRelativeToWarpPointAtTime(const UMotionWarpingBaseAdapter& WarpingAdapter, const UAnimSequenceBase* Animation, float Time, const FTransform& WarpPointTransform);
+
+	/** Extract bone transform from animation at a given time */
+	UFUNCTION(BlueprintCallable, Category = "Motion Warping")
+	static void ExtractBoneTransformFromAnimationAtTime(const UAnimInstance* AnimInstance, const UAnimSequenceBase* Animation, float Time, bool bExtractRootMotion, FName BoneName, bool bLocalSpace, FTransform& OutTransform);
 };
 
 
