Fix to allow motion warping of rotations to account for montage play rate

diff --git a/Engine/Plugins/Animation/MotionWarping/Source/MotionWarping/Private/RootMotionModifier.cpp b/Engine/Plugins/Animation/MotionWarping/Source/MotionWarping/Private/RootMotionModifier.cpp
--- a/Engine/Plugins/Animation/MotionWarping/Source/MotionWarping/Private/RootMotionModifier.cpp
+++ b/Engine/Plugins/Animation/MotionWarping/Source/MotionWarping/Private/RootMotionModifier.cpp
@@ -213,6 +213,7 @@ void URootMotionModifier::Update(const FMotionWarpingUpdateContext& Context)
 	PreviousPosition = Context.PreviousPosition;
 	CurrentPosition = Context.CurrentPosition;
 	Weight = Context.Weight;
+	PlayRate = Context.PlayRate;
 
 	// Mark for removal if the animation already passed the warping window
 	if (PreviousPosition >= EndTime)
@@ -461,13 +462,14 @@ FQuat URootMotionModifier_Warp::WarpRotation(const FTransform& RootMotionDelta,
 	}
 
 	const float TimeRemaining = (EndTime - PreviousPosition) * WarpRotationTimeMultiplier;
-	const float Alpha = FMath::Clamp(DeltaSeconds / TimeRemaining, 0.f, 1.f);
+	const float PlayRateAdjustedDeltaSeconds = DeltaSeconds * PlayRate;
+	const float Alpha = FMath::Clamp(PlayRateAdjustedDeltaSeconds / TimeRemaining, 0.f, 1.f);
 	FQuat TargetRotThisFrame = FQuat::Slerp(TotalRootMotionRotation, TargetRotation, Alpha);
 
 	if (RotationMethod != EMotionWarpRotationMethod::Slerp)
 	{
 		const float AngleDeltaThisFrame = TotalRootMotionRotation.AngularDistance(TargetRotThisFrame);
-		const float MaxAngleDelta = FMath::Abs(FMath::DegreesToRadians(DeltaSeconds * WarpMaxRotationRate));
+		const float MaxAngleDelta = FMath::Abs(FMath::DegreesToRadians(PlayRateAdjustedDeltaSeconds * WarpMaxRotationRate));
 		const float TotalAngleDelta = TotalRootMotionRotation.AngularDistance(TargetRotation);
 		if (RotationMethod == EMotionWarpRotationMethod::ConstantRate && (TotalAngleDelta <= MaxAngleDelta))
 		{
diff --git a/Engine/Plugins/Animation/MotionWarping/Source/MotionWarping/Public/RootMotionModifier.h b/Engine/Plugins/Animation/MotionWarping/Source/MotionWarping/Public/RootMotionModifier.h
--- a/Engine/Plugins/Animation/MotionWarping/Source/MotionWarping/Public/RootMotionModifier.h
+++ b/Engine/Plugins/Animation/MotionWarping/Source/MotionWarping/Public/RootMotionModifier.h
@@ -96,6 +96,10 @@ public:
 	UPROPERTY(BlueprintReadOnly, Category = "Defaults")
 	float Weight = 0.f;
 
+	/** Current play rate of the animation */
+	UPROPERTY(BlueprintReadOnly, Category = "Defaults")
+	float PlayRate = 1.f;
+
 	/** Character owner transform at the time this modifier becomes active */
 	UPROPERTY(Transient, BlueprintReadOnly, Category = "Defaults")
 	FTransform StartTransform;
