[MotionWarping] Fix motion warping rotation offset not applied if rotation type is set to facing. Add opposite facing rotation type.

[RN]

diff --git a/Engine/Plugins/Animation/MotionWarping/Source/MotionWarping/Private/RootMotionModifier.cpp b/Engine/Plugins/Animation/MotionWarping/Source/MotionWarping/Private/RootMotionModifier.cpp
--- a/Engine/Plugins/Animation/MotionWarping/Source/MotionWarping/Private/RootMotionModifier.cpp
+++ b/Engine/Plugins/Animation/MotionWarping/Source/MotionWarping/Private/RootMotionModifier.cpp
@@ -350,6 +350,9 @@ void URootMotionModifier_Warp::Update(const FMotionWarpingUpdateContext& Context
 		// Get the warp point sent by the game
 		FTransform WarpPointTransformGame = WarpTargetPtr->GetTargetTrasform();
 
+		// Cache the rotation offset to later apply when rotation warping 
+		RotationOffset = WarpTargetPtr->RotationOffset.Quaternion();
+		
 		// Initialize our target transform (where the root should end at the end of the window) with the warp point sent by the game
 		FTransform TargetTransform = WarpPointTransformGame;
 
@@ -421,12 +424,21 @@ FQuat URootMotionModifier_Warp::GetTargetRotation() const
 			{
 				const FTransform& ActorTransform = ActorOwner->GetActorTransform();
 				const FVector ToSyncPoint = (CachedTargetTransform.GetLocation() - ActorTransform.GetLocation()).GetSafeNormal2D();
-				return FRotationMatrix::MakeFromXZ(ToSyncPoint, FVector::UpVector).ToQuat();
+				return RotationOffset * FRotationMatrix::MakeFromXZ(ToSyncPoint, FVector::UpVector).ToQuat();
 			}
 			break;
 		
 		case EMotionWarpRotationType::OppositeDefault:
 			return FRotationMatrix::MakeFromXZ(CachedTargetTransform.GetRotation().GetForwardVector() * -1, CachedTargetTransform.GetRotation().GetUpVector()).ToQuat();
+
+		case EMotionWarpRotationType::OppositeFacing:
+			if (const AActor* ActorOwner = GetActorOwner())
+			{
+				const FTransform& ActorTransform = ActorOwner->GetActorTransform();
+				const FVector ToSyncPoint = (ActorTransform.GetLocation() - CachedTargetTransform.GetLocation()).GetSafeNormal2D();
+				return RotationOffset * FRotationMatrix::MakeFromXZ(ToSyncPoint, FVector::UpVector).ToQuat();
+			}
+			break;
 		
 		default:
 			checkNoEntry();
diff --git a/Engine/Plugins/Animation/MotionWarping/Source/MotionWarping/Public/RootMotionModifier.h b/Engine/Plugins/Animation/MotionWarping/Source/MotionWarping/Public/RootMotionModifier.h
--- a/Engine/Plugins/Animation/MotionWarping/Source/MotionWarping/Public/RootMotionModifier.h
+++ b/Engine/Plugins/Animation/MotionWarping/Source/MotionWarping/Public/RootMotionModifier.h
@@ -281,6 +281,9 @@ enum class EMotionWarpRotationType : uint8
 
 	/** Character rotates to match the rotation of the sync point, rotated by half of full angle around up axis (180 degrees yaw) */
 	OppositeDefault,
+
+	/** Character rotates to face the sync point with their back. Essentially facing + 180 degrees yaw. */
+	OppositeFacing,
 };
 
 UENUM(BlueprintType)
@@ -416,6 +419,9 @@ protected:
 	/** Cached of the offset from the warp target. Used to calculate the final target transform when a warp target is defined in the animation */
 	TOptional<FTransform> CachedOffsetFromWarpPoint;
 
+	/** Rotation offset defined in warp target */
+	FQuat RotationOffset = FQuat::Identity;
+	
 	bool bWarpingPaused = false;
 	bool bRootMotionPaused = false;
 };
