Motion Warping - Fixed bug causing incorrect target calculation when a warp target is created from a SceneComponent with the Follow option enabled and an AnimProvider is defined in the warping window.

diff --git a/Engine/Plugins/Animation/MotionWarping/Source/MotionWarping/Private/RootMotionModifier.cpp b/Engine/Plugins/Animation/MotionWarping/Source/MotionWarping/Private/RootMotionModifier.cpp
--- a/Engine/Plugins/Animation/MotionWarping/Source/MotionWarping/Private/RootMotionModifier.cpp
+++ b/Engine/Plugins/Animation/MotionWarping/Source/MotionWarping/Private/RootMotionModifier.cpp
@@ -103,7 +103,7 @@ FTransform FMotionWarpingTarget::GetTargetTrasform() const
 		}
 
 		RecalculateOffset(Transform);
-		return Transform;
+		return FTransform(Transform.GetRotation(), Transform.GetLocation());
 	}
 
 	return FTransform(Rotation, Location);
