Add a custom make bp function for FMotionWarpingTarget to ensure that the custom ctor is called to calculate the target transforms when a component is passed as an input

diff --git a/Engine/Plugins/Animation/MotionWarping/Source/MotionWarping/Private/MotionWarpingFunctionLibrary.cpp b/Engine/Plugins/Animation/MotionWarping/Source/MotionWarping/Private/MotionWarpingFunctionLibrary.cpp
new file mode 100644
--- /dev/null
+++ b/Engine/Plugins/Animation/MotionWarping/Source/MotionWarping/Private/MotionWarpingFunctionLibrary.cpp
@@ -0,0 +1,26 @@
+// Copyright Epic Games, Inc. All Rights Reserved.
+#include "MotionWarpingFunctionLibrary.h"
+
+#include UE_INLINE_GENERATED_CPP_BY_NAME(MotionWarpingFunctionLibrary)
+
+FMotionWarpingTarget UMotionWarpingFunctionLibrary::MakeMotionWarpingTarget(const FName Name, const FVector Location, const FRotator Rotation, const USceneComponent* Component, FName BoneName, bool bFollowComponent, EWarpTargetLocationOffsetDirection LocationOffsetDirection, const AActor* AvatarActor, const FVector LocationOffset, const FRotator RotationOffset)
+{
+	if (Component)
+	{
+		return FMotionWarpingTarget(Name, Component, BoneName, bFollowComponent, LocationOffsetDirection, AvatarActor, LocationOffset, RotationOffset);
+	}
+	else
+	{
+		FMotionWarpingTarget Result = FMotionWarpingTarget();
+
+		// Only certain arguments are valid when a component isn't specified
+		Result.Name = Name;
+		Result.Location = Location;
+		Result.Rotation = Rotation;
+		Result.Component = nullptr;
+		Result.BoneName = NAME_None;
+		Result.bFollowComponent = false;
+
+		return Result;
+	}
+}
diff --git a/Engine/Plugins/Animation/MotionWarping/Source/MotionWarping/Public/MotionWarpingFunctionLibrary.h b/Engine/Plugins/Animation/MotionWarping/Source/MotionWarping/Public/MotionWarpingFunctionLibrary.h
new file mode 100644
--- /dev/null
+++ b/Engine/Plugins/Animation/MotionWarping/Source/MotionWarping/Public/MotionWarpingFunctionLibrary.h
@@ -0,0 +1,25 @@
+// Copyright Epic Games, Inc. All Rights Reserved.
+
+#pragma once
+
+#include "Kismet/BlueprintFunctionLibrary.h"
+
+#include "MotionWarpingFunctionLibrary.generated.h"
+
+#define UE_API MOTIONWARPING_API
+
+/**
+ * Motion Warping Function Library
+ */
+UCLASS(MinimalAPI, meta = (ScriptName = "MotionWarpingLibrary"))
+class UMotionWarpingFunctionLibrary : public UBlueprintFunctionLibrary
+{
+	GENERATED_BODY()
+
+public:
+	UFUNCTION(BlueprintPure, Category = "Motion Warping", meta = (NativeMakeFunc))
+	static UE_API FMotionWarpingTarget MakeMotionWarpingTarget(const FName Name, const FVector Location, const FRotator Rotation, const USceneComponent* Component, FName BoneName, bool bFollowComponent, EWarpTargetLocationOffsetDirection LocationOffsetDirection, const AActor* AvatarActor, const FVector LocationOffset, const FRotator RotationOffset);
+
+};
+
+#undef UE_API
diff --git a/Engine/Plugins/Animation/MotionWarping/Source/MotionWarping/Public/RootMotionModifier.h b/Engine/Plugins/Animation/MotionWarping/Source/MotionWarping/Public/RootMotionModifier.h
--- a/Engine/Plugins/Animation/MotionWarping/Source/MotionWarping/Public/RootMotionModifier.h
+++ b/Engine/Plugins/Animation/MotionWarping/Source/MotionWarping/Public/RootMotionModifier.h
@@ -170,11 +170,11 @@ private:
 UENUM(BlueprintType)
 enum class EWarpTargetLocationOffsetDirection : uint8
 {
-	/* Location Offset's X is target component's forward vector. */
+	/* LocationOffset.X translates the target transform along the target component's forward vector. */
 	TargetsForwardVector,
-	/* Location Offset's X is vector from owner to target. */
+	/* LocationOffset.X translates the target transform along the vector from the target to the owning actor. */
 	VectorFromTargetToOwner,
-	/* Location Offset's X is world space X. */
+	/* LocationOffset moves the target transform in world space */
 	WorldSpace
 };
 
@@ -182,7 +182,7 @@ enum class EWarpTargetLocationOffsetDirection : uint8
 ///////////////////////////////////////////////////////////////
 
 /** Represents a point of alignment in the world */
-USTRUCT(BlueprintType)
+USTRUCT(BlueprintType, meta = (HasNativeMake = "/Script/MotionWarping.MotionWarpingFunctionLibrary.MakeMotionWarpingTarget"))
 struct FMotionWarpingTarget
 {
 	GENERATED_BODY()
@@ -207,11 +207,13 @@ struct FMotionWarpingTarget
 	UPROPERTY(EditAnywhere, BlueprintReadWrite, Category = "Defaults")
 	FName BoneName;
 
-	/** Whether the target transform calculated from a component and an optional bone should be updated during the warp */
+	/** Whether the target transform calculated from a component and an optional bone should be updated during the warp 
+	Only relevant when the warp target is created from a component  */
 	UPROPERTY(EditAnywhere, BlueprintReadWrite, Category = "Defaults")
 	bool bFollowComponent;
 
-	/** Direction of offset context vector */
+	/** Determines how LocationOffset is applied by FMotionWarpingTarget when calculating a target transform
+	Only relevant when the warp target is created from a component */
 	UPROPERTY(EditAnywhere, BlueprintReadWrite, Category = "Defaults")
 	EWarpTargetLocationOffsetDirection LocationOffsetDirection = EWarpTargetLocationOffsetDirection::TargetsForwardVector;
 
