Fix for motion warping switchoff condition to use the correct target transform when component following ends

diff --git a/Engine/Plugins/Animation/MotionWarping/Source/MotionWarping/Private/MotionWarpingComponent.cpp b/Engine/Plugins/Animation/MotionWarping/Source/MotionWarping/Private/MotionWarpingComponent.cpp
--- a/Engine/Plugins/Animation/MotionWarping/Source/MotionWarping/Private/MotionWarpingComponent.cpp
+++ b/Engine/Plugins/Animation/MotionWarping/Source/MotionWarping/Private/MotionWarpingComponent.cpp
@@ -397,16 +397,8 @@ void UMotionWarpingComponent::UpdateSwitchOffConditions()
 					if (WarpTargets[i].bFollowComponent)
 					{
 						WarpTargets[i].bFollowComponent = false;
-
-						if (const USceneComponent* Component = WarpTargets[i].Component.Get())
-						{
-							WarpTargets[i].Location = Component->GetComponentLocation();
-							WarpTargets[i].Rotation = Component->GetComponentRotation();
-						}
-						else
-						{
-							UE_LOG(LogMotionWarping, Warning, TEXT("Switch off condition on warp target %s failed. Trying to cancel follow component, but component isn't valid."), *WarpTargets[i].Name.ToString());
-						}
+						WarpTargets[i].Location = Condition->GetTargetLocation();
+						WarpTargets[i].Rotation = Condition->GetTargetRotation();
 					}
 					break;
 				case ESwitchOffConditionEffect::CancelWarping:
diff --git a/Engine/Plugins/Animation/MotionWarping/Source/MotionWarping/Private/MotionWarpingSwitchOffCondition.cpp b/Engine/Plugins/Animation/MotionWarping/Source/MotionWarping/Private/MotionWarpingSwitchOffCondition.cpp
--- a/Engine/Plugins/Animation/MotionWarping/Source/MotionWarping/Private/MotionWarpingSwitchOffCondition.cpp
+++ b/Engine/Plugins/Animation/MotionWarping/Source/MotionWarping/Private/MotionWarpingSwitchOffCondition.cpp
@@ -37,6 +37,23 @@ FVector UMotionWarpingSwitchOffCondition::GetTargetLocation() const
 	return TargetActor->GetActorLocation();
 }
 
+FRotator UMotionWarpingSwitchOffCondition::GetTargetRotation() const
+{
+	if (bUseWarpTargetAsTargetLocation)
+	{
+		if (MotionWarpingTarget)
+		{
+			return MotionWarpingTarget->Rotator();
+		}
+
+		UE_LOG(LogMotionWarping, Warning, TEXT("Switch off condition is set to use warp target as target rotation, "
+			"however warp target appears to be null. Make sure warp target is correctly passed, "
+			"otherwise switch off condition will use target actor rotation."))
+	}
+
+	return TargetActor->GetActorRotation();
+}
+
 void UMotionWarpingSwitchOffCondition::SetWarpTargetForDestination(const FMotionWarpingTarget* InMotionWarpingTarget)
 {
 	MotionWarpingTarget = InMotionWarpingTarget;
diff --git a/Engine/Plugins/Animation/MotionWarping/Source/MotionWarping/Private/RootMotionModifier.cpp b/Engine/Plugins/Animation/MotionWarping/Source/MotionWarping/Private/RootMotionModifier.cpp
--- a/Engine/Plugins/Animation/MotionWarping/Source/MotionWarping/Private/RootMotionModifier.cpp
+++ b/Engine/Plugins/Animation/MotionWarping/Source/MotionWarping/Private/RootMotionModifier.cpp
@@ -120,16 +120,23 @@ void FMotionWarpingTarget::CacheOffset(const FTransform& InTransform)
 	
 	if (LocationOffsetDirection == EWarpTargetLocationOffsetDirection::VectorFromTargetToOwner)
 	{
-		const FVector ContextVector = (AvatarActor->GetActorLocation() - InTransform.GetLocation()).GetSafeNormal();
-		const FVector RightVector = -FVector::CrossProduct(ContextVector, FVector::UpVector);
+		if (AvatarActor.IsValid())
+		{
+			const FVector ContextVector = (AvatarActor->GetActorLocation() - InTransform.GetLocation()).GetSafeNormal();
+			const FVector RightVector = -FVector::CrossProduct(ContextVector, FVector::UpVector);
+
+			if (bCacheForwardOffset)
+			{
+				CachedForwardOffset = (ContextVector * LocationOffset.X);
+			}
 
-		if (bCacheForwardOffset)
+			CachedRightOffset = RightVector * LocationOffset.Y;
+			CachedUpOffset = -FVector::CrossProduct(RightVector, ContextVector) * LocationOffset.Z;
+		}
+		else
 		{
-			CachedForwardOffset = (ContextVector * LocationOffset.X);
+			UE_LOG(LogMotionWarping, Warning, TEXT("Motion warping offset is set to VectorFromTargetToOwner but avator actor is invalid"))
 		}
-		
-		CachedRightOffset = RightVector * LocationOffset.Y;
-		CachedUpOffset = -FVector::CrossProduct(RightVector, ContextVector) * LocationOffset.Z;
 	}
 }
 
diff --git a/Engine/Plugins/Animation/MotionWarping/Source/MotionWarping/Public/MotionWarpingSwitchOffCondition.h b/Engine/Plugins/Animation/MotionWarping/Source/MotionWarping/Public/MotionWarpingSwitchOffCondition.h
--- a/Engine/Plugins/Animation/MotionWarping/Source/MotionWarping/Public/MotionWarpingSwitchOffCondition.h
+++ b/Engine/Plugins/Animation/MotionWarping/Source/MotionWarping/Public/MotionWarpingSwitchOffCondition.h
@@ -83,6 +83,12 @@ public:
 	 */
 	UE_API FVector GetTargetLocation() const;
 
+	/**
+	* If bUseWarpTargetAsTargetLocation is true, this will return target FMotionWarpingTarget rotation.
+	* If bUseWarpTargetAsTargetLocation is false, this will return target actor rotation.
+	*/
+	UE_API FRotator GetTargetRotation() const;
+
 	/**
 	 * Set warp target as context for calculating switch off condition result
 	 * if bUseWarpTargetAsTargetLocation is set to true
