Make Motion Warping switchoff conditions and warp target direction mode experimental as these APIs will change in future

diff --git a/Engine/Plugins/Animation/MotionWarping/Source/MotionWarping/Public/MotionWarpingComponent.h b/Engine/Plugins/Animation/MotionWarping/Source/MotionWarping/Public/MotionWarpingComponent.h
--- a/Engine/Plugins/Animation/MotionWarping/Source/MotionWarping/Public/MotionWarpingComponent.h
+++ b/Engine/Plugins/Animation/MotionWarping/Source/MotionWarping/Public/MotionWarpingComponent.h
@@ -216,7 +216,7 @@ public:
 	UFUNCTION(BlueprintCallable, Category = "Motion Warping")
 	UE_API int32 RemoveWarpTargets(const TArray<FName>& WarpTargetNames);
 
-	UFUNCTION(BlueprintCallable, Category="Motion Warping")
+	UFUNCTION(BlueprintCallable, Category="Motion Warping|Experimental")
 	UE_API void AddSwitchOffCondition(FName WarpTargetName, UMotionWarpingSwitchOffCondition* Condition);
 	
 	UE_API void RemoveSwitchOffConditions(FName WarpTargetName);
@@ -236,7 +236,7 @@ protected:
 	UPROPERTY(Transient, Replicated)
 	TArray<FMotionWarpingTarget> WarpTargets;
 		
-	UPROPERTY(Transient)
+	UPROPERTY(Transient, Experimental)
     TArray<FSwitchOffConditionData> SwitchOffConditions;
 
 #if !(UE_BUILD_SHIPPING || UE_BUILD_TEST)
diff --git a/Engine/Plugins/Animation/MotionWarping/Source/MotionWarping/Public/MotionWarpingSwitchOffCondition.h b/Engine/Plugins/Animation/MotionWarping/Source/MotionWarping/Public/MotionWarpingSwitchOffCondition.h
--- a/Engine/Plugins/Animation/MotionWarping/Source/MotionWarping/Public/MotionWarpingSwitchOffCondition.h
+++ b/Engine/Plugins/Animation/MotionWarping/Source/MotionWarping/Public/MotionWarpingSwitchOffCondition.h
@@ -10,7 +10,7 @@
 class UMotionWarpingSwitchOffCondition;
 struct FMotionWarpingTarget;
 
-USTRUCT(BlueprintType)
+USTRUCT(BlueprintType, Experimental)
 struct FSwitchOffConditionData
 {
 	GENERATED_BODY()
@@ -39,7 +39,7 @@ struct FSwitchOffConditionData
 };
 
 /**Result of switch off condition.*/
-UENUM(BlueprintType, Category = "MotionWarping")
+UENUM(BlueprintType, Category = "MotionWarping", meta = (Experimental))
 enum class ESwitchOffConditionEffect : uint8
 {
 	/**Changes associated motion warping target from component to a location of this component
@@ -56,7 +56,7 @@ enum class ESwitchOffConditionEffect : uint8
 	PauseRootMotion UMETA(DisplayName = "Pause Root Motion")
 };
 
-UCLASS(MinimalAPI, Abstract, EditInlineNew, ClassGroup = MotionWarping)
+UCLASS(MinimalAPI, Abstract, EditInlineNew, Experimental, ClassGroup = MotionWarping)
 class UMotionWarpingSwitchOffCondition : public UObject
 {
 	GENERATED_BODY()
@@ -118,14 +118,14 @@ protected:
 	const FMotionWarpingTarget* MotionWarpingTarget = nullptr;
 };
 
-UENUM(BlueprintType, Category = "MotionWarping")
+UENUM(BlueprintType, Category = "MotionWarping", meta = (Experimental))
 enum class ESwitchOffConditionDistanceOp : uint8
 {
 	LessThan,
 	GreaterThan
 };
 
-UENUM(BlueprintType, Category = "MotionWarping")
+UENUM(BlueprintType, Category = "MotionWarping", meta = (Experimental))
 enum class ESwitchOffConditionDistanceAxesType : uint8
 {
 	AllAxes,
@@ -133,7 +133,7 @@ enum class ESwitchOffConditionDistanceAxesType : uint8
 	OnlyZAxis,
 };
 
-UCLASS()
+UCLASS(Experimental)
 class UMotionWarpingSwitchOffDistanceCondition : public UMotionWarpingSwitchOffCondition
 {
 	GENERATED_BODY()
@@ -180,14 +180,14 @@ protected:
 };
 
 
-UENUM(BlueprintType, Category = "MotionWarping")
+UENUM(BlueprintType, Category = "MotionWarping", meta = (Experimental))
 enum class ESwitchOffConditionAngleOp : uint8
 {
 	LessThan,
 	GreaterThan
 };
 
-UCLASS()
+UCLASS(Experimental)
 class UMotionWarpingSwitchOffAngleToTargetCondition : public UMotionWarpingSwitchOffCondition
 {
 	GENERATED_BODY()
@@ -207,7 +207,7 @@ public:
 	 * @param InTargetActor Target Actor to use as a context for this Switch Off Condition if Use Warp Target As Location parameter is false.
 	 * @return Created Switch Off Condition. This Switch Off Condition can be added to Motion Warping Component with Add Switch Off Condition node.
 	 */
-	UFUNCTION(BlueprintCallable, Category="Motion Warping", meta=(AdvancedDisplay="bInUseWarpTargetAsTargetLocation, InTargetActor"))
+	UFUNCTION(BlueprintCallable, Category="Motion Warping|Experimental", meta=(AdvancedDisplay="bInUseWarpTargetAsTargetLocation, InTargetActor"))
 	static UMotionWarpingSwitchOffAngleToTargetCondition* CreateSwitchOffAngleToTargetCondition(
 		UPARAM(DisplayName = "Owner Actor", ref)		AActor* InOwnerActor,
 		UPARAM(DisplayName = "Effect")					ESwitchOffConditionEffect InEffect,
@@ -233,14 +233,14 @@ protected:
 	bool bIgnoreZAxis;
 };
 
-UENUM(BlueprintType, Category = "MotionWarping")
+UENUM(BlueprintType, Category = "MotionWarping", meta = (Experimental))
 enum class ESwitchOffConditionCompositeOp : uint8
 {
 	Or UMETA(DisplayName = "OR"),
 	And UMETA(DisplayName = "AND")
 };
 
-UCLASS()
+UCLASS(Experimental)
 class UMotionWarpingSwitchOffCompositeCondition : public UMotionWarpingSwitchOffCondition
 {
 	GENERATED_BODY()
@@ -261,7 +261,7 @@ public:
 	 * @param InTargetActor Target Actor to use as a context for this Switch Off Condition if Use Warp Target As Location parameter is false.
 	 * @return Created Switch Off Condition. This Switch Off Condition can be added to Motion Warping Component with Add Switch Off Condition node.
 	 */
-	UFUNCTION(BlueprintCallable, Category="Motion Warping", meta=(AdvancedDisplay="bInUseWarpTargetAsTargetLocation, InTargetActor"))
+	UFUNCTION(BlueprintCallable, Category="Motion Warping|Experimental", meta=(AdvancedDisplay="bInUseWarpTargetAsTargetLocation, InTargetActor"))
 	static UMotionWarpingSwitchOffCompositeCondition* CreateSwitchOffCompositeCondition(
 		UPARAM(DisplayName = "Owner Actor", ref)			AActor* InOwnerActor,
 		UPARAM(DisplayName = "Effect")						ESwitchOffConditionEffect InEffect,
@@ -287,7 +287,7 @@ protected:
 	TObjectPtr<UMotionWarpingSwitchOffCondition> SwitchOffConditionB;
 };
 
-UCLASS(Blueprintable)
+UCLASS(Blueprintable, Experimental)
 class UMotionWarpingSwitchOffBlueprintableCondition : public UMotionWarpingSwitchOffCondition
 {
 	GENERATED_BODY()
@@ -305,7 +305,7 @@ class UMotionWarpingSwitchOffBlueprintableCondition : public UMotionWarpingSwitc
 	 * @param InTargetActor Target Actor to use as a context for this Switch Off Condition if Use Warp Target As Location parameter is false.
 	 * @return Created Switch Off Condition. This Switch Off Condition can be added to Motion Warping Component with Add Switch Off Condition node.
 	 */
-	UFUNCTION(BlueprintCallable, Category="Motion Warping", meta=(AdvancedDisplay="bInUseWarpTargetAsTargetLocation, InTargetActor"))
+	UFUNCTION(BlueprintCallable, Category="Motion Warping|Experimental", meta=(AdvancedDisplay="bInUseWarpTargetAsTargetLocation, InTargetActor"))
 	static UMotionWarpingSwitchOffBlueprintableCondition* CreateSwitchOffBlueprintableCondition(
 		UPARAM(DisplayName = "Owner Actor", ref)					AActor* InOwnerActor,
 		UPARAM(DisplayName = "Effect")								ESwitchOffConditionEffect InEffect,
diff --git a/Engine/Plugins/Animation/MotionWarping/Source/MotionWarping/Public/RootMotionModifier.h b/Engine/Plugins/Animation/MotionWarping/Source/MotionWarping/Public/RootMotionModifier.h
--- a/Engine/Plugins/Animation/MotionWarping/Source/MotionWarping/Public/RootMotionModifier.h
+++ b/Engine/Plugins/Animation/MotionWarping/Source/MotionWarping/Public/RootMotionModifier.h
@@ -167,7 +167,7 @@ private:
 	ERootMotionModifierState State = ERootMotionModifierState::Waiting;
 };
 
-UENUM(BlueprintType)
+UENUM(BlueprintType, meta=(Experimental))
 enum class EWarpTargetLocationOffsetDirection : uint8
 {
 	/* LocationOffset.X translates the target transform along the target component's forward vector. */
@@ -214,7 +214,7 @@ struct FMotionWarpingTarget
 
 	/** Determines how LocationOffset is applied by FMotionWarpingTarget when calculating a target transform
 	Only relevant when the warp target is created from a component */
-	UPROPERTY(EditAnywhere, BlueprintReadWrite, Category = "Defaults")
+	UPROPERTY(EditAnywhere, BlueprintReadWrite, Category = Experimental)
 	EWarpTargetLocationOffsetDirection LocationOffsetDirection = EWarpTargetLocationOffsetDirection::TargetsForwardVector;
 
 	/** Optional static location offset. Only relevant when the warp target is created from a component */
