MotionWarping - adding URootMotionModifier_Blueprintable, blueprintable overridable URootMotionModifier

diff --git a/Engine/Plugins/Animation/MotionWarping/Source/MotionWarping/Private/RootMotionModifier.cpp b/Engine/Plugins/Animation/MotionWarping/Source/MotionWarping/Private/RootMotionModifier.cpp
--- a/Engine/Plugins/Animation/MotionWarping/Source/MotionWarping/Private/RootMotionModifier.cpp
+++ b/Engine/Plugins/Animation/MotionWarping/Source/MotionWarping/Private/RootMotionModifier.cpp
@@ -670,3 +670,21 @@ URootMotionModifier_Scale* URootMotionModifier_Scale::AddRootMotionModifierScale
 
 	return nullptr;
 }
+
+
+// URootMotionModifier_Blueprintable
+///////////////////////////////////////////////////////////////
+
+FTransform URootMotionModifier_Blueprintable::ProcessRootMotion(const FTransform& InRootMotion, float DeltaSeconds)
+{
+	check(IsInGameThread()); // checking since BP_ProcessRootMotion is not threadsafe
+	if (const UMotionWarpingBaseAdapter* Adapter = GetOwnerAdapter())
+	{
+		if (const ACharacter* CharacterOwner = Cast<ACharacter>(Adapter->GetActor()))
+		{
+			return BP_ProcessRootMotion(CharacterOwner, InRootMotion, DeltaSeconds);
+		}
+	}
+
+	return InRootMotion;
+}
\ No newline at end of file
diff --git a/Engine/Plugins/Animation/MotionWarping/Source/MotionWarping/Public/RootMotionModifier.h b/Engine/Plugins/Animation/MotionWarping/Source/MotionWarping/Public/RootMotionModifier.h
--- a/Engine/Plugins/Animation/MotionWarping/Source/MotionWarping/Public/RootMotionModifier.h
+++ b/Engine/Plugins/Animation/MotionWarping/Source/MotionWarping/Public/RootMotionModifier.h
@@ -475,4 +475,22 @@ public:
 		UPARAM(DisplayName = "Scale") FVector InScale = FVector(1.f));
 };
 
+// URootMotionModifier_Blueprintable
+///////////////////////////////////////////////////////////////
+
+UCLASS(Experimental, MinimalAPI, Blueprintable, meta = (DisplayName = "Blueprintable"))
+class URootMotionModifier_Blueprintable : public URootMotionModifier
+{
+	GENERATED_BODY()
+
+public:
+
+	URootMotionModifier_Blueprintable(const FObjectInitializer& ObjectInitializer) : Super(ObjectInitializer) {}
+
+	virtual FTransform ProcessRootMotion(const FTransform& InRootMotion, float DeltaSeconds) override;
+
+	UFUNCTION(BlueprintPure, BlueprintImplementableEvent, meta=(DisplayName = "Process Root Motion"), Category = "Settings")
+	UE_API FTransform BP_ProcessRootMotion(const ACharacter* CharacterOwner, const FTransform& InRootMotion, float DeltaSeconds) const;
+};
+
 #undef UE_API
