[MotionWarping] Exposing AxesType parameter for CreateSwitchOffDistanceCondition BP node

diff --git a/Engine/Plugins/Animation/MotionWarping/Source/MotionWarping/Private/MotionWarpingSwitchOffCondition.cpp b/Engine/Plugins/Animation/MotionWarping/Source/MotionWarping/Private/MotionWarpingSwitchOffCondition.cpp
--- a/Engine/Plugins/Animation/MotionWarping/Source/MotionWarping/Private/MotionWarpingSwitchOffCondition.cpp
+++ b/Engine/Plugins/Animation/MotionWarping/Source/MotionWarping/Private/MotionWarpingSwitchOffCondition.cpp
@@ -89,12 +89,13 @@ bool UMotionWarpingSwitchOffCondition::IsConditionValid() const
 	return true;
 }
 
-UMotionWarpingSwitchOffDistanceCondition* UMotionWarpingSwitchOffDistanceCondition::CreateSwitchOffDistanceCondition(AActor* InOwnerActor, ESwitchOffConditionEffect InEffect, ESwitchOffConditionDistanceOp InOperator, float InDistance, bool InbUseWarpTargetAsTargetLocation, AActor* InTargetActor)
+UMotionWarpingSwitchOffDistanceCondition* UMotionWarpingSwitchOffDistanceCondition::CreateSwitchOffDistanceCondition(AActor* InOwnerActor, ESwitchOffConditionEffect InEffect, ESwitchOffConditionDistanceOp InOperator, ESwitchOffConditionDistanceAxesType InAxesType, float InDistance, bool InbUseWarpTargetAsTargetLocation, AActor* InTargetActor)
 {
 	UMotionWarpingSwitchOffDistanceCondition* SwitchOffDistanceCondition = NewObject<UMotionWarpingSwitchOffDistanceCondition>();
 	SwitchOffDistanceCondition->OwnerActor = InOwnerActor;
 	SwitchOffDistanceCondition->Effect = InEffect;
 	SwitchOffDistanceCondition->Operator = InOperator;
+	SwitchOffDistanceCondition->AxesType = InAxesType;
 	SwitchOffDistanceCondition->Distance = InDistance;
 	SwitchOffDistanceCondition->bUseWarpTargetAsTargetLocation = InbUseWarpTargetAsTargetLocation;
 	SwitchOffDistanceCondition->TargetActor = InTargetActor;
diff --git a/Engine/Plugins/Animation/MotionWarping/Source/MotionWarping/Public/MotionWarpingSwitchOffCondition.h b/Engine/Plugins/Animation/MotionWarping/Source/MotionWarping/Public/MotionWarpingSwitchOffCondition.h
--- a/Engine/Plugins/Animation/MotionWarping/Source/MotionWarping/Public/MotionWarpingSwitchOffCondition.h
+++ b/Engine/Plugins/Animation/MotionWarping/Source/MotionWarping/Public/MotionWarpingSwitchOffCondition.h
@@ -155,6 +155,7 @@ public:
 	 * @param InOwnerActor Actor owning Motion Warping Component to which this switch off condition will be added.
 	 * @param InEffect What should happen if switch off condition is true
 	 * @param InOperator Distance comparison operator
+	 * @param InAxesType Specifier of the axes the distance should be calculated for
 	 * @param InDistance Distance to compare to
 	 * @param InbUseWarpTargetAsTargetLocation Should use corresponding warp target as context for this Switch Off Condition.
 	 * If set to false, will use TargetActor parameter
@@ -166,6 +167,7 @@ public:
 		UPARAM(DisplayName = "Owner Actor", ref)		AActor* InOwnerActor,
 		UPARAM(DisplayName = "Effect")					ESwitchOffConditionEffect InEffect,
 		UPARAM(DisplayName = "Operator")				ESwitchOffConditionDistanceOp InOperator,
+		UPARAM(DisplayName = "AxesType")				ESwitchOffConditionDistanceAxesType InAxesType,
 		UPARAM(DisplayName = "Distance")				float InDistance,
 		UPARAM(DisplayName = "UseWarpTargetAsLocation") bool InbUseWarpTargetAsTargetLocation = true,
 		UPARAM(DisplayName = "TargetActor")				AActor* InTargetActor = nullptr);
