MotionWarping - exposing URootMotionModifier_PrecomputedWarp as Blueprintable

diff --git a/Engine/Plugins/Animation/MotionWarping/Source/MotionWarping/Public/RootMotionModifier_PrecomputedWarp.h b/Engine/Plugins/Animation/MotionWarping/Source/MotionWarping/Public/RootMotionModifier_PrecomputedWarp.h
--- a/Engine/Plugins/Animation/MotionWarping/Source/MotionWarping/Public/RootMotionModifier_PrecomputedWarp.h
+++ b/Engine/Plugins/Animation/MotionWarping/Source/MotionWarping/Public/RootMotionModifier_PrecomputedWarp.h
@@ -68,7 +68,7 @@ enum class EPrecomputedWarpUpdateMode
 	Relative,
 };
 
-UCLASS(Experimental, MinimalAPI, meta = (DisplayName = "Precomputed Warp", Tooltip="(Experimental) Precomputed warp computes the full path to alignment on the first frame of execution and then on following frames just applies that root motion.  Only Suitable for stationary targets."))
+UCLASS(Experimental, Blueprintable, MinimalAPI, meta = (DisplayName = "Precomputed Warp", Tooltip="(Experimental) Precomputed warp computes the full path to alignment on the first frame of execution and then on following frames just applies that root motion.  Only Suitable for stationary targets."))
 class URootMotionModifier_PrecomputedWarp : public URootMotionModifier_Warp
 {
 	GENERATED_BODY()
