MotionWarping: stop using the warp point's rotation offset when in facing mode.
Add an additional rotation offset in the root motion modifier that allows rotational offset to be specified per animation. This avoids conflicting rotations when warp points are used by multiple different actions.
Note that any setup using a rotational offset with facing mode will need to be manually adjusted (this was introduced in 5.7)
Deprecate EMotionWarpRotationType::OppositeFacing and EMotionWarpRotationType::OppositeDefault since they can be achieved with an additional rotational offset of 180 degrees (existing data will be automatically converted)

#Synced-CL 49661098

diff --git a/Engine/Plugins/Animation/MotionWarping/Source/MotionWarping/Private/RootMotionModifier.cpp b/Engine/Plugins/Animation/MotionWarping/Source/MotionWarping/Private/RootMotionModifier.cpp
--- a/Engine/Plugins/Animation/MotionWarping/Source/MotionWarping/Private/RootMotionModifier.cpp
+++ b/Engine/Plugins/Animation/MotionWarping/Source/MotionWarping/Private/RootMotionModifier.cpp
@@ -74,9 +74,9 @@ void URootMotionModifier_Warp::Serialize(FArchive& Ar)
 {
 	// Handle change of default blend type
 	Ar.UsingCustomVersion(FFortniteMainBranchObjectVersion::GUID);
+	const int32 CustomVersion = Ar.CustomVer(FFortniteMainBranchObjectVersion::GUID);
 	if (Ar.IsLoading())
 	{
-		const int32 CustomVersion = Ar.CustomVer(FFortniteMainBranchObjectVersion::GUID);
 		if (CustomVersion < FFortniteMainBranchObjectVersion::ChangeDefaultAlphaBlendType)
 		{
 			// Switch the default back to Linear so old data remains the same
@@ -84,8 +84,30 @@ void URootMotionModifier_Warp::Serialize(FArchive& Ar)
 			AddTranslationEasingFunc = EAlphaBlendOption::Linear;
 		}
 	}
-	
+
+	// Now load so any other deprecation fixup can work properly
 	Super::Serialize(Ar);
+
+	if (Ar.IsLoading())
+	{
+		PRAGMA_DISABLE_DEPRECATION_WARNINGS
+		if (CustomVersion < FFortniteMainBranchObjectVersion::RootMotionModifierRotationOffset)
+		{
+			if (RotationType == EMotionWarpRotationType::OppositeDefault)
+			{
+				RotationType = EMotionWarpRotationType::Default;
+				AdditionalRotationOffset.Yaw += 180.0f;
+				AdditionalRotationOffset.Normalize();
+			}
+			else if ( RotationType == EMotionWarpRotationType::OppositeFacing )
+			{
+				RotationType = EMotionWarpRotationType::Facing;
+				AdditionalRotationOffset.Yaw += 180.0f;
+				AdditionalRotationOffset.Normalize();
+			}
+		}
+		PRAGMA_ENABLE_DEPRECATION_WARNINGS
+	}
 }
 
 FTransform FMotionWarpingTarget::GetTargetTrasform() const
@@ -424,26 +446,14 @@ FQuat URootMotionModifier_Warp::GetTargetRotation() const
 	switch (RotationType)
 	{
 		case EMotionWarpRotationType::Default:
-			return CachedTargetTransform.GetRotation();
+			return FQuat::MakeFromRotator(AdditionalRotationOffset) * CachedTargetTransform.GetRotation();
 		
 		case EMotionWarpRotationType::Facing:
 			if (const AActor* ActorOwner = GetActorOwner())
 			{
 				const FTransform& ActorTransform = ActorOwner->GetActorTransform();
 				const FVector ToSyncPoint = (CachedTargetTransform.GetLocation() - ActorTransform.GetLocation()).GetSafeNormal2D();
-				return RotationOffset * FRotationMatrix::MakeFromXZ(ToSyncPoint, FVector::UpVector).ToQuat();
-			}
-			break;
-		
-		case EMotionWarpRotationType::OppositeDefault:
-			return FRotationMatrix::MakeFromXZ(CachedTargetTransform.GetRotation().GetForwardVector() * -1, CachedTargetTransform.GetRotation().GetUpVector()).ToQuat();
-
-		case EMotionWarpRotationType::OppositeFacing:
-			if (const AActor* ActorOwner = GetActorOwner())
-			{
-				const FTransform& ActorTransform = ActorOwner->GetActorTransform();
-				const FVector ToSyncPoint = (ActorTransform.GetLocation() - CachedTargetTransform.GetLocation()).GetSafeNormal2D();
-				return RotationOffset * FRotationMatrix::MakeFromXZ(ToSyncPoint, FVector::UpVector).ToQuat();
+				return FQuat::MakeFromRotator(AdditionalRotationOffset) * FRotationMatrix::MakeFromXZ(ToSyncPoint, FVector::UpVector).ToQuat();
 			}
 			break;
 		
diff --git a/Engine/Plugins/Animation/MotionWarping/Source/MotionWarping/Public/RootMotionModifier.h b/Engine/Plugins/Animation/MotionWarping/Source/MotionWarping/Public/RootMotionModifier.h
--- a/Engine/Plugins/Animation/MotionWarping/Source/MotionWarping/Public/RootMotionModifier.h
+++ b/Engine/Plugins/Animation/MotionWarping/Source/MotionWarping/Public/RootMotionModifier.h
@@ -282,10 +282,10 @@ enum class EMotionWarpRotationType : uint8
 	Facing,
 
 	/** Character rotates to match the rotation of the sync point, rotated by half of full angle around up axis (180 degrees yaw) */
-	OppositeDefault,
+	OppositeDefault UE_DEPRECATED(5.8, "Use AdditionalRotationOffset instead") UMETA(Hidden),
 
 	/** Character rotates to face the sync point with their back. Essentially facing + 180 degrees yaw. */
-	OppositeFacing,
+	OppositeFacing UE_DEPRECATED(5.8, "Use AdditionalRotationOffset instead") UMETA(Hidden),
 };
 
 UENUM(BlueprintType)
@@ -378,6 +378,10 @@ public:
 	UPROPERTY(EditAnywhere, BlueprintReadWrite, Category = "Config")
 	bool bSubtractRemainingRootMotion = false;
 
+	/** Applied as a final rotational offset on the warp target */
+	UPROPERTY(EditAnywhere, BlueprintReadWrite, Category = "Config")
+	FRotator AdditionalRotationOffset = FRotator::ZeroRotator;
+
 	/**
 	 * Allow to modify how fast the rotation is warped.
 	 * e.g if the window duration is 2sec and this is 0.5, the target rotation will be reached in 1sec instead of 2sec
diff --git a/Engine/Source/Runtime/Core/Public/UObject/FortniteMainBranchObjectVersions.inl b/Engine/Source/Runtime/Core/Public/UObject/FortniteMainBranchObjectVersions.inl
--- a/Engine/Source/Runtime/Core/Public/UObject/FortniteMainBranchObjectVersions.inl
+++ b/Engine/Source/Runtime/Core/Public/UObject/FortniteMainBranchObjectVersions.inl
@@ -736,5 +736,8 @@ FORTNITE_MAIN_VERSION(CollisionFilter64Bit, 240)
 // Add the concept of UAFAssetData
 FORTNITE_MAIN_VERSION(UAFAssetData, 241)
 
+// Add an additional rotation offset to root motion modifier and deprecate the 'Opposite' rotation methods
+FORTNITE_MAIN_VERSION(RootMotionModifierRotationOffset, 242)
+
 // -----<new versions can be added above this line>-------------------------------------------------
 #undef FORTNITE_MAIN_VERSION
