Fixed dynamic delegate bindings to non-const member functions with const pointers.  This is a const-correctness violation which will soon become a warning (and later, error), as is already the case for regular delegates.

diff --git a/Engine/Plugins/Animation/MotionWarping/Source/MotionWarping/Private/AnimNotifyState_MotionWarping.cpp b/Engine/Plugins/Animation/MotionWarping/Source/MotionWarping/Private/AnimNotifyState_MotionWarping.cpp
--- a/Engine/Plugins/Animation/MotionWarping/Source/MotionWarping/Private/AnimNotifyState_MotionWarping.cpp
+++ b/Engine/Plugins/Animation/MotionWarping/Source/MotionWarping/Private/AnimNotifyState_MotionWarping.cpp
@@ -25,17 +25,17 @@ void UAnimNotifyState_MotionWarping::OnBecomeRelevant(UMotionWarpingComponent* M
 	{
 		if (!RootMotionModifierNew->OnActivateDelegate.IsBound())
 		{
-			RootMotionModifierNew->OnActivateDelegate.BindDynamic(this, &UAnimNotifyState_MotionWarping::OnRootMotionModifierActivate);
+			RootMotionModifierNew->OnActivateDelegate.BindDynamic(const_cast<UAnimNotifyState_MotionWarping*>(this), &UAnimNotifyState_MotionWarping::OnRootMotionModifierActivate);
 		}
 
 		if (!RootMotionModifierNew->OnUpdateDelegate.IsBound())
 		{
-			RootMotionModifierNew->OnUpdateDelegate.BindDynamic(this, &UAnimNotifyState_MotionWarping::OnRootMotionModifierUpdate);
+			RootMotionModifierNew->OnUpdateDelegate.BindDynamic(const_cast<UAnimNotifyState_MotionWarping*>(this), &UAnimNotifyState_MotionWarping::OnRootMotionModifierUpdate);
 		}
 
 		if (!RootMotionModifierNew->OnDeactivateDelegate.IsBound())
 		{
-			RootMotionModifierNew->OnDeactivateDelegate.BindDynamic(this, &UAnimNotifyState_MotionWarping::OnRootMotionModifierDeactivate);
+			RootMotionModifierNew->OnDeactivateDelegate.BindDynamic(const_cast<UAnimNotifyState_MotionWarping*>(this), &UAnimNotifyState_MotionWarping::OnRootMotionModifierDeactivate);
 		}
 	}
 }
diff --git a/Engine/Plugins/AudioGameplayVolume/Source/AudioGameplayVolume/Private/AudioGameplayVolume.cpp b/Engine/Plugins/AudioGameplayVolume/Source/AudioGameplayVolume/Private/AudioGameplayVolume.cpp
--- a/Engine/Plugins/AudioGameplayVolume/Source/AudioGameplayVolume/Private/AudioGameplayVolume.cpp
+++ b/Engine/Plugins/AudioGameplayVolume/Source/AudioGameplayVolume/Private/AudioGameplayVolume.cpp
@@ -161,8 +161,8 @@ void AAudioGameplayVolume::AddProxy() const
 {
 	if (AGVComponent)
 	{
-		AGVComponent->OnProxyEnter.AddUniqueDynamic(this, &AAudioGameplayVolume::OnListenerEnter);
-		AGVComponent->OnProxyExit.AddUniqueDynamic(this, &AAudioGameplayVolume::OnListenerExit);
+		AGVComponent->OnProxyEnter.AddUniqueDynamic(const_cast<AAudioGameplayVolume*>(this), &AAudioGameplayVolume::OnListenerEnter);
+		AGVComponent->OnProxyExit.AddUniqueDynamic(const_cast<AAudioGameplayVolume*>(this), &AAudioGameplayVolume::OnListenerExit);
 		AGVComponent->Activate();
 	}
 }
diff --git a/Engine/Plugins/Experimental/Animation/ContextualAnimation/Source/ContextualAnimation/Private/ContextualAnimSceneActorComponent.cpp b/Engine/Plugins/Experimental/Animation/ContextualAnimation/Source/ContextualAnimation/Private/ContextualAnimSceneActorComponent.cpp
--- a/Engine/Plugins/Experimental/Animation/ContextualAnimation/Source/ContextualAnimation/Private/ContextualAnimSceneActorComponent.cpp
+++ b/Engine/Plugins/Experimental/Animation/ContextualAnimation/Source/ContextualAnimation/Private/ContextualAnimSceneActorComponent.cpp
@@ -987,7 +987,7 @@ void UContextualAnimSceneActorComponent::SetIgnoreCollisionWithActor(AActor& Act
 		//  so we explicitly remove its reference, which unblocks it from being cleaned up by garbage collection
 		if (bValue)
 		{
-			Actor.OnEndPlay.AddUniqueDynamic(this, &UContextualAnimSceneActorComponent::OnIgnoredActorEndPlay);
+			Actor.OnEndPlay.AddUniqueDynamic(const_cast<UContextualAnimSceneActorComponent*>(this), &UContextualAnimSceneActorComponent::OnIgnoredActorEndPlay);
 		}
 		else
 		{
diff --git a/Engine/Source/Runtime/Engine/Private/WorldPartition/WorldPartitionRuntimeLevelStreamingCell.cpp b/Engine/Source/Runtime/Engine/Private/WorldPartition/WorldPartitionRuntimeLevelStreamingCell.cpp
--- a/Engine/Source/Runtime/Engine/Private/WorldPartition/WorldPartitionRuntimeLevelStreamingCell.cpp
+++ b/Engine/Source/Runtime/Engine/Private/WorldPartition/WorldPartitionRuntimeLevelStreamingCell.cpp
@@ -684,8 +684,8 @@ UWorldPartitionLevelStreamingDynamic* UWorldPartitionRuntimeLevelStreamingCell::
 
 	if (LevelStreaming)
 	{
-		LevelStreaming->OnLevelShown.AddUniqueDynamic(this, &UWorldPartitionRuntimeLevelStreamingCell::OnLevelShown);
-		LevelStreaming->OnLevelHidden.AddUniqueDynamic(this, &UWorldPartitionRuntimeLevelStreamingCell::OnLevelHidden);
+		LevelStreaming->OnLevelShown.AddUniqueDynamic(const_cast<UWorldPartitionRuntimeLevelStreamingCell*>(this), &UWorldPartitionRuntimeLevelStreamingCell::OnLevelShown);
+		LevelStreaming->OnLevelHidden.AddUniqueDynamic(const_cast<UWorldPartitionRuntimeLevelStreamingCell*>(this), &UWorldPartitionRuntimeLevelStreamingCell::OnLevelHidden);
 	}
 
 	return LevelStreaming;
