Get the correct compact pose bone index rather than mesh bone index
#Synced-CL 51330716

diff --git a/Engine/Plugins/Animation/MotionWarping/Source/MotionWarping/Private/MotionWarpingComponent.cpp b/Engine/Plugins/Animation/MotionWarping/Source/MotionWarping/Private/MotionWarpingComponent.cpp
--- a/Engine/Plugins/Animation/MotionWarping/Source/MotionWarping/Private/MotionWarpingComponent.cpp
+++ b/Engine/Plugins/Animation/MotionWarping/Source/MotionWarping/Private/MotionWarpingComponent.cpp
@@ -258,20 +258,22 @@ void UMotionWarpingUtilities::ExtractBoneTransformFromAnimationAtTime(const UAni
 	{
 		FMemMark Mark(FMemStack::Get());
 
-		const int32 BoneIndex = AnimInstance->GetRequiredBones().GetPoseBoneIndexForBoneName(BoneName);
-		if (BoneIndex != INDEX_NONE)
+		const int32 MeshBoneIndex = AnimInstance->GetRequiredBones().GetPoseBoneIndexForBoneName(BoneName);
+		const FCompactPoseBoneIndex PoseBoneIndex = AnimInstance->GetRequiredBones().MakeCompactPoseIndex(FMeshPoseBoneIndex(MeshBoneIndex));
+
+		if (PoseBoneIndex != INDEX_NONE)
 		{
 			if (bLocalSpace)
 			{
 				FCompactPose Pose;
 				ExtractLocalSpacePose(Animation, AnimInstance->GetRequiredBones(), Time, bExtractRootMotion, Pose);
-				OutTransform = Pose[FCompactPoseBoneIndex(BoneIndex)];
+				OutTransform = Pose[PoseBoneIndex];
 			}
 			else
 			{
 				FCSPose<FCompactPose> Pose;
 				UMotionWarpingUtilities::ExtractComponentSpacePose(Animation, AnimInstance->GetRequiredBones(), Time, bExtractRootMotion, Pose);
-				OutTransform = Pose.GetComponentSpaceTransform(FCompactPoseBoneIndex(BoneIndex));
+				OutTransform = Pose.GetComponentSpaceTransform(PoseBoneIndex);
 			}
 		}
 	}
