Fixing edge case in RootMotionModifier_SkewWarp for small motion warps with large offsets between mesh and object origins

#Synced-CL 51953934

diff --git a/Engine/Plugins/Animation/MotionWarping/Source/MotionWarping/Private/RootMotionModifier_SkewWarp.cpp b/Engine/Plugins/Animation/MotionWarping/Source/MotionWarping/Private/RootMotionModifier_SkewWarp.cpp
--- a/Engine/Plugins/Animation/MotionWarping/Source/MotionWarping/Private/RootMotionModifier_SkewWarp.cpp
+++ b/Engine/Plugins/Animation/MotionWarping/Source/MotionWarping/Private/RootMotionModifier_SkewWarp.cpp
@@ -189,7 +189,13 @@ FTransform URootMotionModifier_SkewWarp::ProcessRootMotion(const FTransform& InR
 				const FTransform MeshTransform = FTransform(MeshRotationOffset, MeshTranslationOffset) * OwnerAdapter->GetCurrentTransform();
 				TargetLocation = MeshTransform.InverseTransformPositionNoScale(TargetLocation) - RootMotionRemainingAfterNotify.GetTranslation();
 
+				if (bIgnoreZAxis)
+				{
+					TargetLocation.Z = 0;
+				}
+
 				FVector WarpedTranslation = WarpTranslation(FTransform::Identity, DeltaTranslation, TotalTranslation, TargetLocation) + ExtraRootMotion.GetLocation();
+
 				if (MaxSpeedClampRatio > 0.0f)
 				{
 					if (bIgnoreZAxis)
