MotionWarping - fix for FAlignmentNotifyInstance::GetWeight and URootMotionModifier_PrecomputedWarp::GetWeight failing to calculate the proper weight when StartTimeWithRatio is greater than EndTimeWithRatio because of the ActualStartTime

#Synced-CL 52486799

diff --git a/Engine/Plugins/Animation/MotionWarping/Source/MotionWarping/Private/RootMotionModifier_PrecomputedWarp.cpp b/Engine/Plugins/Animation/MotionWarping/Source/MotionWarping/Private/RootMotionModifier_PrecomputedWarp.cpp
--- a/Engine/Plugins/Animation/MotionWarping/Source/MotionWarping/Private/RootMotionModifier_PrecomputedWarp.cpp
+++ b/Engine/Plugins/Animation/MotionWarping/Source/MotionWarping/Private/RootMotionModifier_PrecomputedWarp.cpp
@@ -102,8 +102,16 @@ float URootMotionModifier_PrecomputedWarp::GetWeight(float CurrentTime, const FP
 	{
 		const float StartTimeWithRatio = FMath::Max(FMath::Lerp(StartTime, EndTime, WarpCurve.StartRatio), ActualStartTime);
 		const float EndTimeWithRatio = FMath::Lerp(StartTime, EndTime, WarpCurve.EndRatio);
-			
-		const float CurrentRelativeTime = FMath::Clamp((CurrentTime-StartTimeWithRatio)/(EndTimeWithRatio-StartTimeWithRatio), 0,1);
+		
+		float CurrentRelativeTime;
+		if (StartTimeWithRatio < EndTimeWithRatio - UE_SMALL_NUMBER)
+		{
+			CurrentRelativeTime = FMath::Clamp((CurrentTime - StartTimeWithRatio) / (EndTimeWithRatio - StartTimeWithRatio), 0, 1);
+		}
+		else
+		{
+			CurrentRelativeTime = 1.f;
+		}
 		
 		switch(WarpCurve.CurveType)
 		{
@@ -161,175 +169,170 @@ FTransform URootMotionModifier_PrecomputedWarp::ProcessRootMotion(const FTransfo
 		
 		TargetTransform = FTransform(GetTargetRotation() * MeshOffsetRotation, GetTargetLocation());
 		
-		{
-			static const FBox UnitBox(FVector(-10, -10, -10), FVector(10, 10, 10));
-			UE_VLOG_OBOX(VLogObject, "Alignment", Display, UnitBox, TargetTransform.ToMatrixWithScale(), FColor::Blue, TEXT(""));
-		
-			// get alignment transform
-			const float PredictionDelta = RoundedEndTime - CurrentTime;
+		// get alignment transform
+		const float PredictionDelta = RoundedEndTime - CurrentTime;
 
-			TargetTransform = AlignmentNotify->AlignOffset * TargetTransform;
+		TargetTransform = AlignmentNotify->AlignOffset * TargetTransform;
 
-			if (AlignmentNotify->bForceTargetTransformUpright)
-			{
-				FRotator Rotation = TargetTransform.GetRotation().Rotator();
-				Rotation.Roll = 0;
-				Rotation.Pitch = 0;
-				TargetTransform.SetRotation(Rotation.Quaternion());
-			}
+		if (AlignmentNotify->bForceTargetTransformUpright)
+		{
+			FRotator Rotation = TargetTransform.GetRotation().Rotator();
+			Rotation.Roll = 0;
+			Rotation.Pitch = 0;
+			TargetTransform.SetRotation(Rotation.Quaternion());
+		}
 		
-			StartingRootTransform = CurrentTransform;
+		static const FBox UnitBox(FVector(-10, -10, -10), FVector(10, 10, 10));
+		UE_VLOG_OBOX(VLogObject, "Alignment", Display, UnitBox, TargetTransform.ToMatrixWithScale(), FColor::Blue, TEXT(""));
+
+		StartingRootTransform = CurrentTransform;
 		
-			// extract root motion trajectory todo: this should be cached and reused
+		// extract root motion trajectory todo: this should be cached and reused
 		
-			int NumFrames = 1 + (RoundedEndTime - ActualStartTime) / SamplingFrameTime;
+		int NumFrames = 1 + (RoundedEndTime - ActualStartTime) / SamplingFrameTime;
 
-			if (NumFrames <= 0)
-			{
-				return InRootMotion;
-			}
+		if (NumFrames <= 0)
+		{
+			return InRootMotion;
+		}
 
-			AnimTrajectoryData.Trajectory.SetNum(NumFrames);
-			AnimTrajectoryData.TranslationCurve.SetNum(NumFrames);
-			AnimTrajectoryData.RotationCurve.SetNum(NumFrames);
+		AnimTrajectoryData.Trajectory.SetNum(NumFrames);
+		AnimTrajectoryData.TranslationCurve.SetNum(NumFrames);
+		AnimTrajectoryData.RotationCurve.SetNum(NumFrames);
 		
-			float SteeringAngleThreshold = FMath::DegreesToRadians(AlignmentNotify->SteeringSettings.AngleThreshold);
+		float SteeringAngleThreshold = FMath::DegreesToRadians(AlignmentNotify->SteeringSettings.AngleThreshold);
 		
-			FTransform PredictedTransform;
-			float PredictionTime = CurrentTime;
+		FTransform PredictedTransform;
+		float PredictionTime = CurrentTime;
 			
-			AnimTrajectoryData.Trajectory[0] = PredictedTransform;
-			PredictionTime += SamplingFrameTime;
-			AnimTrajectoryData.TranslationCurve[0] = 0;
+		AnimTrajectoryData.Trajectory[0] = PredictedTransform;
+		PredictionTime += SamplingFrameTime;
+		AnimTrajectoryData.TranslationCurve[0] = 0;
 
-			for(int i=1;i<NumFrames; i++)
-			{
-				FTransform RootMotionThisFrame = ExtractRootMotionHelper(AnimationAsset, PredictionTime, SamplingFrameTime, false);
+		for(int i=1;i<NumFrames; i++)
+		{
+			FTransform RootMotionThisFrame = ExtractRootMotionHelper(AnimationAsset, PredictionTime, SamplingFrameTime, false);
 			
-				PredictedTransform = RootMotionThisFrame * PredictedTransform;
-				AnimTrajectoryData.Trajectory[i] = PredictedTransform;
-				PredictionTime += SamplingFrameTime;
+			PredictedTransform = RootMotionThisFrame * PredictedTransform;
+			AnimTrajectoryData.Trajectory[i] = PredictedTransform;
+			PredictionTime += SamplingFrameTime;
 
-				AnimTrajectoryData.TranslationCurve[i] = RootMotionThisFrame.GetTranslation().Length() + AnimTrajectoryData.TranslationCurve[i-1]; 
-				AnimTrajectoryData.RotationCurve[i] = fabs(RootMotionThisFrame.GetRotation().GetAngle()) + AnimTrajectoryData.RotationCurve[i-1]; 
-			}
+			AnimTrajectoryData.TranslationCurve[i] = RootMotionThisFrame.GetTranslation().Length() + AnimTrajectoryData.TranslationCurve[i-1]; 
+			AnimTrajectoryData.RotationCurve[i] = fabs(RootMotionThisFrame.GetRotation().GetAngle()) + AnimTrajectoryData.RotationCurve[i-1]; 
+		}
 
-			// normalize curves;
-			if (AnimTrajectoryData.TranslationCurve.Last() < UE_SMALL_NUMBER)
-			{
-				AnimTrajectoryData.TranslationCurve.Reset();
-			}
-			else
+		// normalize curves;
+		if (AnimTrajectoryData.TranslationCurve.Last() < UE_SMALL_NUMBER)
+		{
+			AnimTrajectoryData.TranslationCurve.Reset();
+		}
+		else
+		{
+			for(float& Value : AnimTrajectoryData.TranslationCurve)
 			{
-				for(float& Value : AnimTrajectoryData.TranslationCurve)
-				{
-					Value /= AnimTrajectoryData.TranslationCurve.Last();
-				}
+				Value /= AnimTrajectoryData.TranslationCurve.Last();
 			}
+		}
 		
-			if (AnimTrajectoryData.RotationCurve.Last() < UE_SMALL_NUMBER)
-			{
-				AnimTrajectoryData.RotationCurve.Reset();
-			}
-			else
+		if (AnimTrajectoryData.RotationCurve.Last() < UE_SMALL_NUMBER)
+		{
+			AnimTrajectoryData.RotationCurve.Reset();
+		}
+		else
+		{
+			for(float& Value : AnimTrajectoryData.RotationCurve)
 			{
-				for(float& Value : AnimTrajectoryData.RotationCurve)
-				{
-					Value /= AnimTrajectoryData.RotationCurve.Last();
-				}
+				Value /= AnimTrajectoryData.RotationCurve.Last();
 			}
+		}
 
-			FVector UnWarpedPreviousPosition;
-			FVector WarpedPreviousPosition;
+		FVector UnWarpedPreviousPosition;
+		FVector WarpedPreviousPosition;
 
-			if (AlignmentNotify->bEnableSteering && AlignmentNotify->SteeringSettings.bEnableSmoothing)
-			{
-				FilteredSteeringTarget = FQuat::Identity;
-				TargetSmoothingState.Reset();
+		if (AlignmentNotify->bEnableSteering && AlignmentNotify->SteeringSettings.bEnableSmoothing)
+		{
+			FilteredSteeringTarget = FQuat::Identity;
+			TargetSmoothingState.Reset();
 
-				FilteredSteeringTarget = UKismetMathLibrary::QuaternionSpringInterp(FilteredSteeringTarget, FQuat::Identity, TargetSmoothingState,
-												AlignmentNotify->SteeringSettings.SmoothStiffness, AlignmentNotify->SteeringSettings.SmoothDamping, DeltaSeconds, 1, 0, true);
-			}
+			FilteredSteeringTarget = UKismetMathLibrary::QuaternionSpringInterp(FilteredSteeringTarget, FQuat::Identity, TargetSmoothingState,
+											AlignmentNotify->SteeringSettings.SmoothStiffness, AlignmentNotify->SteeringSettings.SmoothDamping, DeltaSeconds, 1, 0, true);
+		}
 
-			WarpedTrajectory.SetNum(AnimTrajectoryData.Trajectory.Num());
+		WarpedTrajectory.SetNum(AnimTrajectoryData.Trajectory.Num());
 			
-			FVector MovementDirection = CurrentTransform.GetRotation().GetRightVector(); // todo, get this from translation over first few frames or based on some setting?
+		FVector MovementDirection = CurrentTransform.GetRotation().GetRightVector(); // todo, get this from translation over first few frames or based on some setting?
 		
-			FTransform InverseLastFrame = AnimTrajectoryData.Trajectory.Last().Inverse();
-			// Translation warping + steering
-			for(int i=0;i<NumFrames; i++)
-			{
-				FTransform TransformFromRoot = AnimTrajectoryData.Trajectory[i] * CurrentTransform;
-				FTransform TransformFromTarget = AnimTrajectoryData.Trajectory[i] * InverseLastFrame * TargetTransform;
+		FTransform InverseLastFrame = AnimTrajectoryData.Trajectory.Last().Inverse();
+		// Translation warping + steering
+		for(int i=0;i<NumFrames; i++)
+		{
+			FTransform TransformFromRoot = AnimTrajectoryData.Trajectory[i] * CurrentTransform;
+			FTransform TransformFromTarget = AnimTrajectoryData.Trajectory[i] * InverseLastFrame * TargetTransform;
 			
-				FVector OldPosition = TransformFromRoot.GetTranslation();
-				FVector UnWarpedDelta = OldPosition - UnWarpedPreviousPosition;
-				UnWarpedPreviousPosition = OldPosition;
+			FVector OldPosition = TransformFromRoot.GetTranslation();
+			FVector UnWarpedDelta = OldPosition - UnWarpedPreviousPosition;
+			UnWarpedPreviousPosition = OldPosition;
 
-				FVector NewPosition = OldPosition;
+			FVector NewPosition = OldPosition;
 				
-				if (AlignmentNotify->bSeparateTranslationCurves)
-				{
-					float InMovementDirectionWeight = GetWeight(ActualStartTime + SamplingFrameTime*i, AlignmentNotify->TranslationWarpingCurve_InMovementDirection);
-					float OutOfMovementDirectionWeight = GetWeight(ActualStartTime + SamplingFrameTime*i, AlignmentNotify->TranslationWarpingCurve_OutOfMovementDirection);
+			if (AlignmentNotify->bSeparateTranslationCurves)
+			{
+				float InMovementDirectionWeight = GetWeight(ActualStartTime + SamplingFrameTime*i, AlignmentNotify->TranslationWarpingCurve_InMovementDirection);
+				float OutOfMovementDirectionWeight = GetWeight(ActualStartTime + SamplingFrameTime*i, AlignmentNotify->TranslationWarpingCurve_OutOfMovementDirection);
 
-					FVector Delta = TransformFromTarget.GetTranslation() - OldPosition;
-					FVector InMovementDirectionDelta = MovementDirection * Delta.Dot(MovementDirection);
-					FVector OutOfMovementDirectionDelta = Delta - InMovementDirectionDelta;
+				FVector Delta = TransformFromTarget.GetTranslation() - OldPosition;
+				FVector InMovementDirectionDelta = MovementDirection * Delta.Dot(MovementDirection);
+				FVector OutOfMovementDirectionDelta = Delta - InMovementDirectionDelta;
 
-					NewPosition = OldPosition + InMovementDirectionDelta * InMovementDirectionWeight + OutOfMovementDirectionDelta * OutOfMovementDirectionWeight;
-				}
-				else
-				{
-					float WarpWeight = GetWeight(ActualStartTime + SamplingFrameTime*i, AlignmentNotify->TranslationWarpingCurve);
-					NewPosition = FMath::Lerp(OldPosition, TransformFromTarget.GetTranslation(), WarpWeight);
-				}
+				NewPosition = OldPosition + InMovementDirectionDelta * InMovementDirectionWeight + OutOfMovementDirectionDelta * OutOfMovementDirectionWeight;
+			}
+			else
+			{
+				float WarpWeight = GetWeight(ActualStartTime + SamplingFrameTime*i, AlignmentNotify->TranslationWarpingCurve);
+				NewPosition = FMath::Lerp(OldPosition, TransformFromTarget.GetTranslation(), WarpWeight);
+			}
 
-				FVector WarpedDelta = NewPosition - WarpedPreviousPosition;
-				WarpedPreviousPosition = NewPosition;
+			FVector WarpedDelta = NewPosition - WarpedPreviousPosition;
+			WarpedPreviousPosition = NewPosition;
 
-				WarpedTrajectory[i].SetTranslation(NewPosition);
-				WarpedTrajectory[i].SetRotation(TransformFromRoot.GetRotation());
+			WarpedTrajectory[i].SetTranslation(NewPosition);
+			WarpedTrajectory[i].SetRotation(TransformFromRoot.GetRotation());
 
-				if (i > 0 && AlignmentNotify->bEnableSteering)
-				{
-					FQuat OldRotation = TransformFromRoot.GetRotation();
-					FQuat DirectionChange = FQuat::FindBetweenVectors(FVector(UnWarpedDelta.X, UnWarpedDelta.Y, 0) , FVector(WarpedDelta.X, WarpedDelta.Y, 0));
+			if (i > 0 && AlignmentNotify->bEnableSteering)
+			{
+				FQuat OldRotation = TransformFromRoot.GetRotation();
+				FQuat DirectionChange = FQuat::FindBetweenVectors(FVector(UnWarpedDelta.X, UnWarpedDelta.Y, 0) , FVector(WarpedDelta.X, WarpedDelta.Y, 0));
 
-					if (AlignmentNotify->SteeringSettings.bEnableSmoothing)
+				if (AlignmentNotify->SteeringSettings.bEnableSmoothing)
+				{
+					if (DirectionChange.GetAngle() < SteeringAngleThreshold)
 					{
-					
-						if (DirectionChange.GetAngle() < SteeringAngleThreshold)
-						{
-							FilteredSteeringTarget = UKismetMathLibrary::QuaternionSpringInterp(FilteredSteeringTarget, DirectionChange, TargetSmoothingState,
-															AlignmentNotify->SteeringSettings.SmoothStiffness, AlignmentNotify->SteeringSettings.SmoothDamping, DeltaSeconds, 1, 0, true);
-						}
-					
-						DirectionChange = FilteredSteeringTarget;
-						WarpedTrajectory[i].SetRotation(OldRotation*DirectionChange);
+						FilteredSteeringTarget = UKismetMathLibrary::QuaternionSpringInterp(FilteredSteeringTarget, DirectionChange, TargetSmoothingState,
+														AlignmentNotify->SteeringSettings.SmoothStiffness, AlignmentNotify->SteeringSettings.SmoothDamping, DeltaSeconds, 1, 0, true);
 					}
-					else
+					
+					DirectionChange = FilteredSteeringTarget;
+					WarpedTrajectory[i].SetRotation(OldRotation*DirectionChange);
+				}
+				else
+				{
+					if (DirectionChange.GetAngle() < SteeringAngleThreshold)
 					{
-						if (DirectionChange.GetAngle() < SteeringAngleThreshold)
-						{
-							WarpedTrajectory[i].SetRotation(OldRotation*DirectionChange);
-						}
+						WarpedTrajectory[i].SetRotation(OldRotation*DirectionChange);
 					}
-				
 				}
 			}
+		}
 		
-			// Rotation warping
-			for(int i=0;i<NumFrames; i++)
-			{
-				float WarpWeight = GetWeight(ActualStartTime + SamplingFrameTime*i, AlignmentNotify->RotationWarpingCurve);
+		// Rotation warping
+		for(int i=0;i<NumFrames; i++)
+		{
+			float WarpWeight = GetWeight(ActualStartTime + SamplingFrameTime*i, AlignmentNotify->RotationWarpingCurve);
 			
-				FQuat OldRotation = WarpedTrajectory[i].GetRotation();
-				FTransform TransformFromTarget = AnimTrajectoryData.Trajectory[i] * InverseLastFrame * TargetTransform;
+			FQuat OldRotation = WarpedTrajectory[i].GetRotation();
+			FTransform TransformFromTarget = AnimTrajectoryData.Trajectory[i] * InverseLastFrame * TargetTransform;
 			
-				WarpedTrajectory[i].SetRotation(FQuat::Slerp(OldRotation, TransformFromTarget.GetRotation(), WarpWeight));
-			}
-		
+			WarpedTrajectory[i].SetRotation(FQuat::Slerp(OldRotation, TransformFromTarget.GetRotation(), WarpWeight));
 		}
 	}
 
diff --git a/Engine/Plugins/Experimental/EvaluationNotifies/Source/Runtime/Private/EvaluationNotifies/AnimNotifyState_Alignment.cpp b/Engine/Plugins/Experimental/EvaluationNotifies/Source/Runtime/Private/EvaluationNotifies/AnimNotifyState_Alignment.cpp
--- a/Engine/Plugins/Experimental/EvaluationNotifies/Source/Runtime/Private/EvaluationNotifies/AnimNotifyState_Alignment.cpp
+++ b/Engine/Plugins/Experimental/EvaluationNotifies/Source/Runtime/Private/EvaluationNotifies/AnimNotifyState_Alignment.cpp
@@ -90,7 +90,15 @@ float FAlignmentNotifyInstance::GetWeight(float CurrentTime, const FAlignmentWar
 		const float StartTimeWithRatio = FMath::Max(FMath::Lerp(StartTime, EndTime, WarpCurve.StartRatio), ActualStartTime);
 		const float EndTimeWithRatio = FMath::Lerp(StartTime, EndTime, WarpCurve.EndRatio);
 			
-		const float CurrentRelativeTime = FMath::Clamp((CurrentTime-StartTimeWithRatio)/(EndTimeWithRatio-StartTimeWithRatio), 0,1);
+		float CurrentRelativeTime;
+		if (StartTimeWithRatio < EndTimeWithRatio - UE_SMALL_NUMBER)
+		{
+			CurrentRelativeTime = FMath::Clamp((CurrentTime - StartTimeWithRatio) / (EndTimeWithRatio - StartTimeWithRatio), 0, 1);
+		}
+		else
+		{
+			CurrentRelativeTime = 1.f;
+		}
 		
 		switch(WarpCurve.CurveType)
 		{
@@ -136,13 +144,9 @@ void FAlignmentNotifyInstance::Update(const UAnimSequenceBase* AnimationAsset, f
 		{
 			TargetTransform = *FoundTransform;
 			
-			static const FBox UnitBox(FVector(-10, -10, -10), FVector(10, 10, 10));
-			UE_VLOG_OBOX(Output.AnimInstanceProxy->GetAnimInstanceObject(), "Alignment", Display, UnitBox, TargetTransform.ToMatrixWithScale(), FColor::Blue, TEXT(""));
-			
 			// get alignment transform
 			const float PredictionDelta = EndTime - CurrentTime ;
 			
-
 			if (AlignBone.HasValidSetup())
 			{
 				// get alignment bone relative to predicted root motion end point, and remove that as an offset to the alignment target
@@ -157,6 +161,10 @@ void FAlignmentNotifyInstance::Update(const UAnimSequenceBase* AnimationAsset, f
 			}
 			
 			TargetTransform = AlignmentNotify->AlignOffset * TargetTransform;
+
+			static const FBox UnitBox(FVector(-10, -10, -10), FVector(10, 10, 10));
+			UE_VLOG_OBOX(Output.AnimInstanceProxy->GetAnimInstanceObject(), "Alignment", Display, UnitBox, TargetTransform.ToMatrixWithScale(), FColor::Blue, TEXT(""));
+			
 			StartingRootTransform = RootBoneTransform;
 			
 			// extract root motion trajectory todo: this should be cached and reused
@@ -281,7 +289,6 @@ void FAlignmentNotifyInstance::Update(const UAnimSequenceBase* AnimationAsset, f
 							WarpedTrajectory[i].SetRotation(OldRotation*DirectionChange);
 						}
 					}
-					
 				}
 			}
 			
@@ -295,7 +302,6 @@ void FAlignmentNotifyInstance::Update(const UAnimSequenceBase* AnimationAsset, f
 				
 				WarpedTrajectory[i].SetRotation(FQuat::Slerp(OldRotation, TransformFromTarget.GetRotation(), Weight));
 			}
-			
 		}
 	}
 	
